Kinect system commands

diff --git a/Engine/Source/Runtime/Core/Public/GenericPlatform/GenericApplicationMessageHandler.h b/Engine/Source/Runtime/Core/Public/GenericPlatform/GenericApplicationMessageHandler.h
--- a/Engine/Source/Runtime/Core/Public/GenericPlatform/GenericApplicationMessageHandler.h
+++ b/Engine/Source/Runtime/Core/Public/GenericPlatform/GenericApplicationMessageHandler.h
@@ -64,6 +64,13 @@ namespace EControllerButtons
 		BackLeft,
 		BackRight,
 
+		// global speech commands
+		GlobalMenu,
+		GlobalView,
+		GlobalPause,
+		GlobalPlay,
+		GlobalBack,
+
 		Invalid,
 	};
 }
diff --git a/Engine/Source/Runtime/InputCore/Classes/InputCoreTypes.h b/Engine/Source/Runtime/InputCore/Classes/InputCoreTypes.h
--- a/Engine/Source/Runtime/InputCore/Classes/InputCoreTypes.h
+++ b/Engine/Source/Runtime/InputCore/Classes/InputCoreTypes.h
@@ -356,6 +356,13 @@ struct INPUTCORE_API EKeys
 	static const FKey Steam_Back_Left;
 	static const FKey Steam_Back_Right;
 
+	// Xbox One global speech commands
+	static const FKey Global_Menu;
+	static const FKey Global_View;
+	static const FKey Global_Pause;
+	static const FKey Global_Play;
+	static const FKey Global_Back;
+
 	static const FKey Invalid;
 
 	static const int32 NUM_TOUCH_KEYS = 10;
diff --git a/Engine/Source/Runtime/InputCore/Private/InputCoreTypes.cpp b/Engine/Source/Runtime/InputCore/Private/InputCoreTypes.cpp
--- a/Engine/Source/Runtime/InputCore/Private/InputCoreTypes.cpp
+++ b/Engine/Source/Runtime/InputCore/Private/InputCoreTypes.cpp
@@ -212,6 +212,13 @@ const FKey EKeys::Gesture_Flick("Gesture_Flick");
 // PS4-specific
 const FKey EKeys::PS4_Special("PS4_Special");
 
+// Xbox One global speech commands
+const FKey EKeys::Global_Menu("Global_Menu");
+const FKey EKeys::Global_View("Global_View");
+const FKey EKeys::Global_Pause("Global_Pause");
+const FKey EKeys::Global_Play("Global_Play");
+const FKey EKeys::Global_Back("Global_Back");
+
 const FKey EKeys::Invalid(NAME_None);
 
 bool EKeys::bInitialized = false;
@@ -424,6 +431,13 @@ void EKeys::Initialize()
 
 	// PS4-specific
 	AddKey(FKeyDetails(EKeys::PS4_Special, LOCTEXT("PS4_Special", "PS4_Special"), FKeyDetails::NotBlueprintBindableKey));
+
+	// Xbox One global speech commands
+	AddKey(FKeyDetails(EKeys::Global_Menu, LOCTEXT("Global_Menu", "Global Menu"), FKeyDetails::GamepadKey));
+	AddKey(FKeyDetails(EKeys::Global_View, LOCTEXT("Global_View", "Global View"), FKeyDetails::GamepadKey));
+	AddKey(FKeyDetails(EKeys::Global_Pause, LOCTEXT("Global_Pause", "Global Pause"), FKeyDetails::GamepadKey));
+	AddKey(FKeyDetails(EKeys::Global_Play, LOCTEXT("Global_Play", "Global Play"), FKeyDetails::GamepadKey));
+	AddKey(FKeyDetails(EKeys::Global_Back, LOCTEXT("Global_Back", "Global Back"), FKeyDetails::GamepadKey));
 }
 
 void EKeys::AddKey(const FKeyDetails& KeyDetails)
diff --git a/Engine/Source/Runtime/Slate/Private/SlateApplication.cpp b/Engine/Source/Runtime/Slate/Private/SlateApplication.cpp
--- a/Engine/Source/Runtime/Slate/Private/SlateApplication.cpp
+++ b/Engine/Source/Runtime/Slate/Private/SlateApplication.cpp
@@ -4022,6 +4022,12 @@ FKey TranslateControllerButtonToKey( EControllerButtons::Type Button )
 	case EControllerButtons::RightStickLeft: Key = EKeys::Gamepad_RightStick_Left; break;
 	case EControllerButtons::RightStickRight: Key = EKeys::Gamepad_RightStick_Right; break;
 
+	case EControllerButtons::GlobalMenu: Key = EKeys::Global_Menu; break;
+	case EControllerButtons::GlobalView: Key = EKeys::Global_View; break;
+	case EControllerButtons::GlobalPause: Key = EKeys::Global_Pause; break;
+	case EControllerButtons::GlobalPlay: Key = EKeys::Global_Play; break;
+	case EControllerButtons::GlobalBack: Key = EKeys::Global_Back; break;
+
 	case EControllerButtons::Invalid: Key = EKeys::Invalid; break;
 	}
 
