Remove unimplemented GetKeyCodeFromKey
Add GetCodesFromKey function that mirrors GetKeyFromCodes

diff --git a/Engine/Source/Runtime/InputCore/Classes/InputCoreTypes.h b/Engine/Source/Runtime/InputCore/Classes/InputCoreTypes.h
--- a/Engine/Source/Runtime/InputCore/Classes/InputCoreTypes.h
+++ b/Engine/Source/Runtime/InputCore/Classes/InputCoreTypes.h
@@ -433,7 +433,7 @@ struct INPUTCORE_API FInputKeyManager
 public:
 	static FInputKeyManager& Get();
 
-	const uint16* GetKeyCodeFromKey( FKey Key ) const;
+	void GetCodesFromKey(const FKey Key, const uint16*& KeyCode, const uint16*& CharCode) const;
 
 	/**
 	 * Retrieves the key mapped to the specified character code.
diff --git a/Engine/Source/Runtime/InputCore/Private/InputCoreTypes.cpp b/Engine/Source/Runtime/InputCore/Private/InputCoreTypes.cpp
--- a/Engine/Source/Runtime/InputCore/Private/InputCoreTypes.cpp
+++ b/Engine/Source/Runtime/InputCore/Private/InputCoreTypes.cpp
@@ -735,4 +735,10 @@ FKey FInputKeyManager::GetKeyFromCodes( const uint16 KeyCode, const uint16 CharC
 	return KeyPtr ? *KeyPtr : EKeys::Invalid;
 }
 
+void FInputKeyManager::GetCodesFromKey(const FKey Key, const uint16*& KeyCode, const uint16*& CharCode) const
+{
+	KeyCode = KeyMapCharToEnum.FindKey(Key);
+	CharCode = KeyMapVirtualToEnum.FindKey(Key);
+}
+
 #undef LOCTEXT_NAMESPACE
