Remove misleading MouseWheelSpin Input mapping

diff --git a/Engine/Source/Runtime/InputCore/Classes/InputCoreTypes.h b/Engine/Source/Runtime/InputCore/Classes/InputCoreTypes.h
--- a/Engine/Source/Runtime/InputCore/Classes/InputCoreTypes.h
+++ b/Engine/Source/Runtime/InputCore/Classes/InputCoreTypes.h
@@ -186,9 +186,6 @@ struct INPUTCORE_API EKeys
 	static const FKey MouseScrollUp;
 	static const FKey MouseScrollDown;
 
-	//@static const FKey todo Slate: Unify MouseScrollUp, Down; Spin.  The viewport clients use Up and Down and slate uses Spin
-	static const FKey MouseWheelSpin;
-
 	static const FKey LeftMouseButton;
 	static const FKey RightMouseButton;
 	static const FKey MiddleMouseButton;
diff --git a/Engine/Source/Runtime/InputCore/Private/InputCoreTypes.cpp b/Engine/Source/Runtime/InputCore/Private/InputCoreTypes.cpp
--- a/Engine/Source/Runtime/InputCore/Private/InputCoreTypes.cpp
+++ b/Engine/Source/Runtime/InputCore/Private/InputCoreTypes.cpp
@@ -15,9 +15,6 @@ const FKey EKeys::MouseY("MouseY");
 const FKey EKeys::MouseScrollUp("MouseScrollUp");
 const FKey EKeys::MouseScrollDown("MouseScrollDown");
 
-//@const FKey EKeys::todo Slate: Unify MouseScrollUp, Down("todo Slate: Unify MouseScrollUp, Down"); Spin.  The viewport clients use Up and Down and slate uses Spin
-const FKey EKeys::MouseWheelSpin("MouseWheelSpin");
-
 const FKey EKeys::LeftMouseButton("LeftMouseButton");
 const FKey EKeys::RightMouseButton("RightMouseButton");
 const FKey EKeys::MiddleMouseButton("MiddleMouseButton");
@@ -234,9 +231,6 @@ void EKeys::Initialize()
 	AddKey(FKeyDetails(EKeys::MouseScrollUp, LOCTEXT("MouseScrollUp", "Mouse Wheel Up"), FKeyDetails::MouseButton));
 	AddKey(FKeyDetails(EKeys::MouseScrollDown, LOCTEXT("MouseScrollDown", "Mouse Wheel Down"), FKeyDetails::MouseButton));
 
-	//@todo Slate: Unify MouseScrollUp, Down, Spin.  The viewport clients use Up and Down and slate uses Spin
-	AddKey(FKeyDetails(EKeys::MouseWheelSpin, LOCTEXT("MouseWheelSpin", "Mouse Wheel Spin"), FKeyDetails::NotBlueprintBindableKey | FKeyDetails::MouseButton));
-
 	AddKey(FKeyDetails(EKeys::LeftMouseButton, LOCTEXT("LeftMouseButton", "Left Mouse Button"), FKeyDetails::MouseButton));
 	AddKey(FKeyDetails(EKeys::RightMouseButton, LOCTEXT("RightMouseButton", "Right Mouse Button"), FKeyDetails::MouseButton));
 	AddKey(FKeyDetails(EKeys::MiddleMouseButton, LOCTEXT("MiddleMouseButton", "Middle Mouse Button"), FKeyDetails::MouseButton));
diff --git a/Engine/Source/Runtime/Slate/Private/Framework/Application/SlateApplication.cpp b/Engine/Source/Runtime/Slate/Private/Framework/Application/SlateApplication.cpp
--- a/Engine/Source/Runtime/Slate/Private/Framework/Application/SlateApplication.cpp
+++ b/Engine/Source/Runtime/Slate/Private/Framework/Application/SlateApplication.cpp
@@ -3617,7 +3617,7 @@ bool FSlateApplication::OnMouseWheel( const float Delta )
 		CurrentCursorPosition, 	
 		CurrentCursorPosition,
 		PressedMouseButtons,
-		EKeys::MouseWheelSpin,
+		EKeys::Invalid,
 		Delta,
 		PlatformApplication->GetModifierKeys()
 		);
@@ -4120,7 +4120,7 @@ bool FSlateApplication::OnTouchGesture( EGestureEvent::Type GestureType, const F
 		CurrentCursorPosition,
 		CurrentCursorPosition,
 		PressedMouseButtons,
-		EKeys::MouseWheelSpin,
+		EKeys::Invalid,
 		MouseWheelDelta,
 		PlatformApplication->GetModifierKeys()
 	);
