Fix bug where keys using the function based naming did not get flagged as axes correctly.

diff --git a/Engine/Source/Runtime/InputCore/Classes/InputCoreTypes.h b/Engine/Source/Runtime/InputCore/Classes/InputCoreTypes.h
--- a/Engine/Source/Runtime/InputCore/Classes/InputCoreTypes.h
+++ b/Engine/Source/Runtime/InputCore/Classes/InputCoreTypes.h
@@ -118,6 +118,15 @@ struct INPUTCORE_API FKeyDetails
 		, bIsBindableInBlueprints((~InKeyFlags & EKeyFlags::NotBlueprintBindableKey) != 0)
 		, AxisType(EInputAxisType::None)
 	{
+		if ((InKeyFlags & EKeyFlags::FloatAxis) != 0)
+		{
+			ensure((InKeyFlags & EKeyFlags::VectorAxis) == 0);
+			AxisType = EInputAxisType::Float;
+		}
+		else if ((InKeyFlags & EKeyFlags::VectorAxis) != 0)
+		{
+			AxisType = EInputAxisType::Vector;
+		}
 	}
 
 	bool IsModifierKey() const { return bIsModifierKey != 0; }
