Make PS4 labels correct for shoulder and trigger buttons when using experimental gamepad labelling.

diff --git a/Engine/Source/Runtime/InputCore/Private/InputCoreTypes.cpp b/Engine/Source/Runtime/InputCore/Private/InputCoreTypes.cpp
--- a/Engine/Source/Runtime/InputCore/Private/InputCoreTypes.cpp
+++ b/Engine/Source/Runtime/InputCore/Private/InputCoreTypes.cpp
@@ -362,9 +362,6 @@ void EKeys::Initialize()
 	AddKey(FKeyDetails(EKeys::Gamepad_RightX, LOCTEXT("Gamepad_RightX", "Gamepad Right X"), FKeyDetails::GamepadKey | FKeyDetails::FloatAxis));
 	AddKey(FKeyDetails(EKeys::Gamepad_RightY, LOCTEXT("Gamepad_RightY", "Gamepad Right Y"), FKeyDetails::GamepadKey | FKeyDetails::FloatAxis));
 
-	AddKey(FKeyDetails(EKeys::Gamepad_LeftTriggerAxis, LOCTEXT("Gamepad_LeftTriggerAxis", "Gamepad Left Trigger Axis"), FKeyDetails::GamepadKey | FKeyDetails::FloatAxis));
-	AddKey(FKeyDetails(EKeys::Gamepad_RightTriggerAxis, LOCTEXT("Gamepad_RightTriggerAxis", "Gamepad Right Trigger Axis"), FKeyDetails::GamepadKey | FKeyDetails::FloatAxis));
-
 	AddKey(FKeyDetails(EKeys::Gamepad_LeftThumbstick, LOCTEXT("Gamepad_LeftThumbstick", "Gamepad Left Thumbstick"), FKeyDetails::GamepadKey));
 	AddKey(FKeyDetails(EKeys::Gamepad_RightThumbstick, LOCTEXT("Gamepad_RightThumbstick", "Gamepad Right Thumbstick"), FKeyDetails::GamepadKey));
 	AddKey(FKeyDetails(EKeys::Gamepad_DPad_Up, LOCTEXT("Gamepad_DPad_Up", "Gamepad DPad Up"), FKeyDetails::GamepadKey));
@@ -372,11 +369,6 @@ void EKeys::Initialize()
 	AddKey(FKeyDetails(EKeys::Gamepad_DPad_Right, LOCTEXT("Gamepad_DPad_Right", "Gamepad DPad Right"), FKeyDetails::GamepadKey));
 	AddKey(FKeyDetails(EKeys::Gamepad_DPad_Left, LOCTEXT("Gamepad_DPad_Left", "Gamepad DPad Left"), FKeyDetails::GamepadKey));
 
-	AddKey(FKeyDetails(EKeys::Gamepad_LeftShoulder, LOCTEXT("Gamepad_LeftShoulder", "Gamepad Left Shoulder"), FKeyDetails::GamepadKey));
-	AddKey(FKeyDetails(EKeys::Gamepad_RightShoulder, LOCTEXT("Gamepad_RightShoulder", "Gamepad Right Shoulder"), FKeyDetails::GamepadKey));
-	AddKey(FKeyDetails(EKeys::Gamepad_LeftTrigger, LOCTEXT("Gamepad_LeftTrigger", "Gamepad Left Trigger"), FKeyDetails::GamepadKey));
-	AddKey(FKeyDetails(EKeys::Gamepad_RightTrigger, LOCTEXT("Gamepad_RightTrigger", "Gamepad Right Trigger"), FKeyDetails::GamepadKey));
-
 	// Virtual key codes used for input axis button press/release emulation
 	AddKey(FKeyDetails(EKeys::Gamepad_LeftStick_Up, LOCTEXT("Gamepad_LeftStick_Up", "Gamepad Left Stick Up"), FKeyDetails::GamepadKey));
 	AddKey(FKeyDetails(EKeys::Gamepad_LeftStick_Down, LOCTEXT("Gamepad_LeftStick_Down", "Gamepad Left Stick Down"), FKeyDetails::GamepadKey));
@@ -397,6 +389,14 @@ void EKeys::Initialize()
 	AddKey(FKeyDetails(EKeys::Gamepad_FaceButton_Left, GetKeyNameDelegate, FKeyDetails::GamepadKey));
 	AddKey(FKeyDetails(EKeys::Gamepad_FaceButton_Top, GetKeyNameDelegate, FKeyDetails::GamepadKey));
 
+	AddKey(FKeyDetails(EKeys::Gamepad_LeftTriggerAxis, GetKeyNameDelegate, FKeyDetails::GamepadKey | FKeyDetails::FloatAxis));
+	AddKey(FKeyDetails(EKeys::Gamepad_RightTriggerAxis, GetKeyNameDelegate, FKeyDetails::GamepadKey | FKeyDetails::FloatAxis));
+
+	AddKey(FKeyDetails(EKeys::Gamepad_LeftShoulder, GetKeyNameDelegate, FKeyDetails::GamepadKey));
+	AddKey(FKeyDetails(EKeys::Gamepad_RightShoulder, GetKeyNameDelegate, FKeyDetails::GamepadKey));
+	AddKey(FKeyDetails(EKeys::Gamepad_LeftTrigger, GetKeyNameDelegate, FKeyDetails::GamepadKey));
+	AddKey(FKeyDetails(EKeys::Gamepad_RightTrigger, GetKeyNameDelegate, FKeyDetails::GamepadKey));
+
 	// Vector axes (FVector, not float)
 	AddKey(FKeyDetails(EKeys::Tilt, LOCTEXT("Tilt", "Tilt"), FKeyDetails::VectorAxis));
 	AddKey(FKeyDetails(EKeys::RotationRate, LOCTEXT("RotationRate", "Rotation Rate"), FKeyDetails::VectorAxis));
@@ -488,6 +488,30 @@ FText EKeys::GetGamepadDisplayName(const FKey Key)
 		{
 			return LOCTEXT("PS4_Gamepad_Special_Right", "Gamepad Options");
 		}
+		else if (Key == EKeys::Gamepad_LeftShoulder)
+		{
+			return LOCTEXT("PS4_Gamepad_LeftShoulder", "Gamepad L1");
+		}
+		else if (Key == EKeys::Gamepad_RightShoulder)
+		{
+			return LOCTEXT("PS4_Gamepad_RightShoulder", "Gamepad R1");
+		}
+		else if (Key == EKeys::Gamepad_LeftTrigger)
+		{
+			return LOCTEXT("PS4_Gamepad_LeftTrigger", "Gamepad L2");
+		}
+		else if (Key == EKeys::Gamepad_RightTrigger)
+		{
+			return LOCTEXT("PS4_Gamepad_RightTrigger", "Gamepad R2");
+		}
+		else if (Key == EKeys::Gamepad_LeftTriggerAxis)
+		{
+			return LOCTEXT("PS4_Gamepad_LeftTriggerAxis", "Gamepad L2 Axis");
+		}
+		else if (Key == EKeys::Gamepad_RightTriggerAxis)
+		{
+			return LOCTEXT("PS4_Gamepad_RightTriggerAxis", "Gamepad R2 Axis");
+		}
 		break;
 
 	case EConsoleForGamepadLabels::XBoxOne:
@@ -515,6 +539,30 @@ FText EKeys::GetGamepadDisplayName(const FKey Key)
 		{
 			return LOCTEXT("XBoxOne_Gamepad_Special_Right", "Gamepad Start");
 		}
+		else if (Key == EKeys::Gamepad_LeftShoulder)
+		{
+			return LOCTEXT("Gamepad_LeftShoulder", "Gamepad Left Shoulder");
+		}
+		else if (Key == EKeys::Gamepad_RightShoulder)
+		{
+			return LOCTEXT("Gamepad_RightShoulder", "Gamepad Right Shoulder");
+		}
+		else if (Key == EKeys::Gamepad_LeftTrigger)
+		{
+			return LOCTEXT("Gamepad_LeftTrigger", "Gamepad Left Trigger");
+		}
+		else if (Key == EKeys::Gamepad_RightTrigger)
+		{
+			return LOCTEXT("Gamepad_RightTrigger", "Gamepad Right Trigger");
+		}
+		else if (Key == EKeys::Gamepad_LeftTriggerAxis)
+		{
+			return LOCTEXT("Gamepad_LeftTriggerAxis", "Gamepad Left Trigger Axis");
+		}
+		else if (Key == EKeys::Gamepad_RightTriggerAxis)
+		{
+			return LOCTEXT("Gamepad_RightTriggerAxis", "Gamepad Right Trigger Axis");
+		}
 		break;
 
 	default:
@@ -542,6 +590,30 @@ FText EKeys::GetGamepadDisplayName(const FKey Key)
 		{
 			return LOCTEXT("Gamepad_Special_Right", "Gamepad Special Right");
 		}
+		else if (Key == EKeys::Gamepad_LeftShoulder)
+		{
+			return LOCTEXT("Gamepad_LeftShoulder", "Gamepad Left Shoulder");
+		}
+		else if (Key == EKeys::Gamepad_RightShoulder)
+		{
+			return LOCTEXT("Gamepad_RightShoulder", "Gamepad Right Shoulder");
+		}
+		else if (Key == EKeys::Gamepad_LeftTrigger)
+		{
+			return LOCTEXT("Gamepad_LeftTrigger", "Gamepad Left Trigger");
+		}
+		else if (Key == EKeys::Gamepad_RightTrigger)
+		{
+			return LOCTEXT("Gamepad_RightTrigger", "Gamepad Right Trigger");
+		}
+		else if (Key == EKeys::Gamepad_LeftTriggerAxis)
+		{
+			return LOCTEXT("Gamepad_LeftTriggerAxis", "Gamepad Left Trigger Axis");
+		}
+		else if (Key == EKeys::Gamepad_RightTriggerAxis)
+		{
+			return LOCTEXT("Gamepad_RightTriggerAxis", "Gamepad Right Trigger Axis");
+		}
 		break;
 	}
 
