#ttp 338022 TappyChicken: Android: Add back button

Added 'Android Back' for specific back button support in Android.
Tested via adding a temp node to Tappy Chicken to make sure it works as expected; which it does.

diff --git a/Engine/Source/Runtime/Core/Private/Android/AndroidMisc.cpp b/Engine/Source/Runtime/Core/Private/Android/AndroidMisc.cpp
--- a/Engine/Source/Runtime/Core/Private/Android/AndroidMisc.cpp
+++ b/Engine/Source/Runtime/Core/Private/Android/AndroidMisc.cpp
@@ -369,6 +369,9 @@ uint32 FAndroidMisc::GetKeyMap( uint16* KeyCodes, FString* KeyNames, uint32 MaxM
 		ADDKEYMAP( AKEYCODE_ALT_RIGHT, TEXT("RightAlt") );
 		//ADDKEYMAP( VK_LWIN, TEXT("LeftCommand") );
 		//ADDKEYMAP( VK_RWIN, TEXT("RightCommand") );
+
+		ADDKEYMAP(AKEYCODE_BACK, TEXT("Android_Back"));
+
 	}
 
 	check(NumMappings < MaxMappings);
diff --git a/Engine/Source/Runtime/InputCore/Classes/InputCoreTypes.h b/Engine/Source/Runtime/InputCore/Classes/InputCoreTypes.h
--- a/Engine/Source/Runtime/InputCore/Classes/InputCoreTypes.h
+++ b/Engine/Source/Runtime/InputCore/Classes/InputCoreTypes.h
@@ -389,6 +389,8 @@ struct INPUTCORE_API EKeys
 	static const FKey Global_Play;
 	static const FKey Global_Back;
 
+	static const FKey Android_Back;
+
 	static const FKey Invalid;
 
 	static const int32 NUM_TOUCH_KEYS = 10;
diff --git a/Engine/Source/Runtime/InputCore/Private/InputCoreTypes.cpp b/Engine/Source/Runtime/InputCore/Private/InputCoreTypes.cpp
--- a/Engine/Source/Runtime/InputCore/Private/InputCoreTypes.cpp
+++ b/Engine/Source/Runtime/InputCore/Private/InputCoreTypes.cpp
@@ -216,6 +216,8 @@ const FKey EKeys::Global_Pause("Global_Pause");
 const FKey EKeys::Global_Play("Global_Play");
 const FKey EKeys::Global_Back("Global_Back");
 
+const FKey EKeys::Android_Back("Android_Back");
+ 
 const FKey EKeys::Invalid(NAME_None);
 
 bool EKeys::bInitialized = false;
@@ -433,6 +435,8 @@ void EKeys::Initialize()
 	AddKey(FKeyDetails(EKeys::Global_Play, LOCTEXT("Global_Play", "Global Play"), FKeyDetails::GamepadKey));
 	AddKey(FKeyDetails(EKeys::Global_Back, LOCTEXT("Global_Back", "Global Back"), FKeyDetails::GamepadKey));
 
+	AddKey(FKeyDetails(EKeys::Android_Back, LOCTEXT("Android_Back", "Android Back"), FKeyDetails::GamepadKey));
+
 	// Initialize the input key manager.  This will cause any additional OEM keys to get added
 	FInputKeyManager::Get();
 }
