Make FKey a BlueprintType

diff --git a/Engine/Source/Runtime/InputCore/Classes/InputCoreTypes.h b/Engine/Source/Runtime/InputCore/Classes/InputCoreTypes.h
--- a/Engine/Source/Runtime/InputCore/Classes/InputCoreTypes.h
+++ b/Engine/Source/Runtime/InputCore/Classes/InputCoreTypes.h
@@ -4,7 +4,7 @@
 
 #include "InputCoreTypes.generated.h"
 
-USTRUCT()
+USTRUCT(BlueprintType)
 struct INPUTCORE_API FKey
 {
 	GENERATED_USTRUCT_BODY()
