Make FKeyDetails DisplayName an Attribute

diff --git a/Engine/Source/Runtime/InputCore/Classes/InputCoreTypes.h b/Engine/Source/Runtime/InputCore/Classes/InputCoreTypes.h
--- a/Engine/Source/Runtime/InputCore/Classes/InputCoreTypes.h
+++ b/Engine/Source/Runtime/InputCore/Classes/InputCoreTypes.h
@@ -90,7 +90,7 @@ struct INPUTCORE_API FKeyDetails
 		NoFlags                 = 0,
 	};
 
-	FKeyDetails(const FKey InKey, const FText InDisplayName, const uint8 InKeyFlags = 0)
+	FKeyDetails(const FKey InKey, const TAttribute<FText>& InDisplayName, const uint8 InKeyFlags = 0)
 		: Key(InKey)
 		, DisplayName(InDisplayName)
 		, bIsModifierKey((InKeyFlags & EKeyFlags::ModifierKey) != 0)
@@ -110,26 +110,6 @@ struct INPUTCORE_API FKeyDetails
 		}
 	}
 
-	FKeyDetails(const FKey InKey, const FGetKeyDisplayNameSignature InGetDisplayNameDelegate, const uint8 InKeyFlags = 0)
-		: Key(InKey)
-		, GetKeyDisplayNameDelegate(InGetDisplayNameDelegate)
-		, bIsModifierKey((InKeyFlags & EKeyFlags::ModifierKey) != 0)
-		, bIsGamepadKey((InKeyFlags & EKeyFlags::GamepadKey) != 0)
-		, bIsMouseButton((InKeyFlags & EKeyFlags::MouseButton) != 0)
-		, bIsBindableInBlueprints((~InKeyFlags & EKeyFlags::NotBlueprintBindableKey) != 0)
-		, AxisType(EInputAxisType::None)
-	{
-		if ((InKeyFlags & EKeyFlags::FloatAxis) != 0)
-		{
-			ensure((InKeyFlags & EKeyFlags::VectorAxis) == 0);
-			AxisType = EInputAxisType::Float;
-		}
-		else if ((InKeyFlags & EKeyFlags::VectorAxis) != 0)
-		{
-			AxisType = EInputAxisType::Vector;
-		}
-	}
-
 	bool IsModifierKey() const { return bIsModifierKey != 0; }
 	bool IsGamepadKey() const { return bIsGamepadKey != 0; }
 	bool IsMouseButton() const { return bIsMouseButton != 0; }
@@ -150,8 +130,7 @@ private:
 
 	FKey  Key;
 	
-	FText DisplayName;
-	FGetKeyDisplayNameSignature GetKeyDisplayNameDelegate;
+	TAttribute<FText> DisplayName;
 
 	int32 bIsModifierKey:1;
 	int32 bIsGamepadKey:1;
diff --git a/Engine/Source/Runtime/InputCore/Private/InputCoreTypes.cpp b/Engine/Source/Runtime/InputCore/Private/InputCoreTypes.cpp
--- a/Engine/Source/Runtime/InputCore/Private/InputCoreTypes.cpp
+++ b/Engine/Source/Runtime/InputCore/Private/InputCoreTypes.cpp
@@ -382,22 +382,20 @@ void EKeys::Initialize()
 	AddKey(FKeyDetails(EKeys::Gamepad_RightStick_Right, LOCTEXT("Gamepad_RightStick_Right", "Gamepad Right Stick Right"), FKeyDetails::GamepadKey));
 	AddKey(FKeyDetails(EKeys::Gamepad_RightStick_Left, LOCTEXT("Gamepad_RightStick_Left", "Gamepad Right Stick Left"), FKeyDetails::GamepadKey));
 
-	FGetKeyDisplayNameSignature GetKeyNameDelegate = FGetKeyDisplayNameSignature::CreateStatic(&EKeys::GetGamepadDisplayName);
+	AddKey(FKeyDetails(EKeys::Gamepad_Special_Left, TAttribute<FText>::Create(TAttribute<FText>::FGetter::CreateStatic(&EKeys::GetGamepadDisplayName, EKeys::Gamepad_Special_Left)), FKeyDetails::GamepadKey));
+	AddKey(FKeyDetails(EKeys::Gamepad_Special_Right, TAttribute<FText>::Create(TAttribute<FText>::FGetter::CreateStatic(&EKeys::GetGamepadDisplayName, EKeys::Gamepad_Special_Right)), FKeyDetails::GamepadKey));
+	AddKey(FKeyDetails(EKeys::Gamepad_FaceButton_Bottom, TAttribute<FText>::Create(TAttribute<FText>::FGetter::CreateStatic(&EKeys::GetGamepadDisplayName, EKeys::Gamepad_FaceButton_Bottom)), FKeyDetails::GamepadKey));
+	AddKey(FKeyDetails(EKeys::Gamepad_FaceButton_Right, TAttribute<FText>::Create(TAttribute<FText>::FGetter::CreateStatic(&EKeys::GetGamepadDisplayName, EKeys::Gamepad_FaceButton_Right)), FKeyDetails::GamepadKey));
+	AddKey(FKeyDetails(EKeys::Gamepad_FaceButton_Left, TAttribute<FText>::Create(TAttribute<FText>::FGetter::CreateStatic(&EKeys::GetGamepadDisplayName, EKeys::Gamepad_FaceButton_Left)), FKeyDetails::GamepadKey));
+	AddKey(FKeyDetails(EKeys::Gamepad_FaceButton_Top, TAttribute<FText>::Create(TAttribute<FText>::FGetter::CreateStatic(&EKeys::GetGamepadDisplayName, EKeys::Gamepad_FaceButton_Top)), FKeyDetails::GamepadKey));
 
-	AddKey(FKeyDetails(EKeys::Gamepad_Special_Left, GetKeyNameDelegate, FKeyDetails::GamepadKey));
-	AddKey(FKeyDetails(EKeys::Gamepad_Special_Right, GetKeyNameDelegate, FKeyDetails::GamepadKey));
-	AddKey(FKeyDetails(EKeys::Gamepad_FaceButton_Bottom, GetKeyNameDelegate, FKeyDetails::GamepadKey));
-	AddKey(FKeyDetails(EKeys::Gamepad_FaceButton_Right, GetKeyNameDelegate, FKeyDetails::GamepadKey));
-	AddKey(FKeyDetails(EKeys::Gamepad_FaceButton_Left, GetKeyNameDelegate, FKeyDetails::GamepadKey));
-	AddKey(FKeyDetails(EKeys::Gamepad_FaceButton_Top, GetKeyNameDelegate, FKeyDetails::GamepadKey));
+	AddKey(FKeyDetails(EKeys::Gamepad_LeftTriggerAxis, TAttribute<FText>::Create(TAttribute<FText>::FGetter::CreateStatic(&EKeys::GetGamepadDisplayName, EKeys::Gamepad_LeftTriggerAxis)), FKeyDetails::GamepadKey | FKeyDetails::FloatAxis));
+	AddKey(FKeyDetails(EKeys::Gamepad_RightTriggerAxis, TAttribute<FText>::Create(TAttribute<FText>::FGetter::CreateStatic(&EKeys::GetGamepadDisplayName, EKeys::Gamepad_RightTriggerAxis)), FKeyDetails::GamepadKey | FKeyDetails::FloatAxis));
 
-	AddKey(FKeyDetails(EKeys::Gamepad_LeftTriggerAxis, GetKeyNameDelegate, FKeyDetails::GamepadKey | FKeyDetails::FloatAxis));
-	AddKey(FKeyDetails(EKeys::Gamepad_RightTriggerAxis, GetKeyNameDelegate, FKeyDetails::GamepadKey | FKeyDetails::FloatAxis));
-
-	AddKey(FKeyDetails(EKeys::Gamepad_LeftShoulder, GetKeyNameDelegate, FKeyDetails::GamepadKey));
-	AddKey(FKeyDetails(EKeys::Gamepad_RightShoulder, GetKeyNameDelegate, FKeyDetails::GamepadKey));
-	AddKey(FKeyDetails(EKeys::Gamepad_LeftTrigger, GetKeyNameDelegate, FKeyDetails::GamepadKey));
-	AddKey(FKeyDetails(EKeys::Gamepad_RightTrigger, GetKeyNameDelegate, FKeyDetails::GamepadKey));
+	AddKey(FKeyDetails(EKeys::Gamepad_LeftShoulder, TAttribute<FText>::Create(TAttribute<FText>::FGetter::CreateStatic(&EKeys::GetGamepadDisplayName, EKeys::Gamepad_LeftShoulder)), FKeyDetails::GamepadKey));
+	AddKey(FKeyDetails(EKeys::Gamepad_RightShoulder, TAttribute<FText>::Create(TAttribute<FText>::FGetter::CreateStatic(&EKeys::GetGamepadDisplayName, EKeys::Gamepad_RightShoulder)), FKeyDetails::GamepadKey));
+	AddKey(FKeyDetails(EKeys::Gamepad_LeftTrigger, TAttribute<FText>::Create(TAttribute<FText>::FGetter::CreateStatic(&EKeys::GetGamepadDisplayName, EKeys::Gamepad_LeftTrigger)), FKeyDetails::GamepadKey));
+	AddKey(FKeyDetails(EKeys::Gamepad_RightTrigger, TAttribute<FText>::Create(TAttribute<FText>::FGetter::CreateStatic(&EKeys::GetGamepadDisplayName, EKeys::Gamepad_RightTrigger)), FKeyDetails::GamepadKey));
 
 	// Vector axes (FVector, not float)
 	AddKey(FKeyDetails(EKeys::Tilt, LOCTEXT("Tilt", "Tilt"), FKeyDetails::VectorAxis));
@@ -686,12 +684,7 @@ bool FKey::IsBindableInBlueprints() const
 
 FText FKeyDetails::GetDisplayName() const
 {
-	if (GetKeyDisplayNameDelegate.IsBound())
-	{
-		return GetKeyDisplayNameDelegate.Execute(Key);
-	}
-
-	return DisplayName;
+	return DisplayName.Get();
 }
 
 FText FKey::GetDisplayName() const
