Support AZERTY keys in FKey system
NOTE: Existing bindings to EKeys::Underscore and EKeys::Quote will no longer work and need to be switched to Hyphen and Apostrophe respectively.  Underscore and Quote now map to the actual keys in their name for binding on AZERTY (and any other keys where they are the unshifted value).

diff --git a/Engine/Source/Runtime/Core/Private/GenericPlatform/GenericPlatformMisc.cpp b/Engine/Source/Runtime/Core/Private/GenericPlatform/GenericPlatformMisc.cpp
--- a/Engine/Source/Runtime/Core/Private/GenericPlatform/GenericPlatformMisc.cpp
+++ b/Engine/Source/Runtime/Core/Private/GenericPlatform/GenericPlatformMisc.cpp
@@ -724,16 +724,32 @@ uint32 FGenericPlatformMisc::GetStandardPrintableKeyMap(uint16* KeyCodes, FStrin
 	ADDKEYMAP( ';', TEXT("Semicolon") );
 	ADDKEYMAP( '=', TEXT("Equals") );
 	ADDKEYMAP( ',', TEXT("Comma") );
-	ADDKEYMAP( '-', TEXT("Underscore") );
+	ADDKEYMAP( '-', TEXT("Hyphen") );
 	ADDKEYMAP( '.', TEXT("Period") );
 	ADDKEYMAP( '/', TEXT("Slash") );
 	ADDKEYMAP( '`', TEXT("Tilde") );
 	ADDKEYMAP( '[', TEXT("LeftBracket") );
 	ADDKEYMAP( '\\', TEXT("Backslash") );
 	ADDKEYMAP( ']', TEXT("RightBracket") );
-	ADDKEYMAP( '\'', TEXT("Quote") );
+	ADDKEYMAP( '\'', TEXT("Apostrophe") );
 	ADDKEYMAP( ' ', TEXT("SpaceBar") );
 
+	// AZERTY Keys
+	ADDKEYMAP( '&', TEXT("Ampersand") );
+	ADDKEYMAP( '*', TEXT("Asterix") );
+	ADDKEYMAP( '^', TEXT("Caret") );
+	ADDKEYMAP( ':', TEXT("Colon") );
+	ADDKEYMAP( '$', TEXT("Dollar") );
+	ADDKEYMAP( '!', TEXT("Exclamation") );
+	ADDKEYMAP( '(', TEXT("LeftParantheses") );
+	ADDKEYMAP( ')', TEXT("RightParantheses") );
+	ADDKEYMAP( '"', TEXT("Quote") );
+	ADDKEYMAP( '_', TEXT("Underscore") );
+	ADDKEYMAP( 224, TEXT("A_AccentGrave") );
+	ADDKEYMAP( 231, TEXT("C_Cedille") );
+	ADDKEYMAP( 233, TEXT("E_AccentAigu") );
+	ADDKEYMAP( 232, TEXT("E_AccentGrave") );
+
 	return NumMappings;
 }
 
diff --git a/Engine/Source/Runtime/InputCore/Classes/InputCoreTypes.h b/Engine/Source/Runtime/InputCore/Classes/InputCoreTypes.h
--- a/Engine/Source/Runtime/InputCore/Classes/InputCoreTypes.h
+++ b/Engine/Source/Runtime/InputCore/Classes/InputCoreTypes.h
@@ -287,14 +287,30 @@ struct INPUTCORE_API EKeys
 	static const FKey Equals;
 	static const FKey Comma;
 	static const FKey Underscore;
+	static const FKey Hyphen;
 	static const FKey Period;
 	static const FKey Slash;
 	static const FKey Tilde;
 	static const FKey LeftBracket;
 	static const FKey Backslash;
 	static const FKey RightBracket;
+	static const FKey Apostrophe;
+
+	static const FKey Ampersand;
+	static const FKey Asterix;
+	static const FKey Caret;
+	static const FKey Colon;
+	static const FKey Dollar;
+	static const FKey Exclamation;
+	static const FKey LeftParantheses;
+	static const FKey RightParantheses;
 	static const FKey Quote;
 
+	static const FKey A_AccentGrave;
+	static const FKey E_AccentGrave;
+	static const FKey E_AccentAigu;
+	static const FKey C_Cedille;
+
 	// Platform Keys
 	// These keys platform specific versions of keys that go by different names.
 	// The delete key is a good example, on Windows Delete is the virtual key for Delete.
diff --git a/Engine/Source/Runtime/InputCore/Private/InputCoreTypes.cpp b/Engine/Source/Runtime/InputCore/Private/InputCoreTypes.cpp
--- a/Engine/Source/Runtime/InputCore/Private/InputCoreTypes.cpp
+++ b/Engine/Source/Runtime/InputCore/Private/InputCoreTypes.cpp
@@ -127,14 +127,31 @@ const FKey EKeys::Semicolon("Semicolon");
 const FKey EKeys::Equals("Equals");
 const FKey EKeys::Comma("Comma");
 const FKey EKeys::Underscore("Underscore");
+const FKey EKeys::Hyphen("Hyphen");
 const FKey EKeys::Period("Period");
 const FKey EKeys::Slash("Slash");
 const FKey EKeys::Tilde("Tilde");
 const FKey EKeys::LeftBracket("LeftBracket");
+const FKey EKeys::LeftParantheses("LeftParantheses");
 const FKey EKeys::Backslash("Backslash");
 const FKey EKeys::RightBracket("RightBracket");
+const FKey EKeys::RightParantheses("RightParantheses");
+const FKey EKeys::Apostrophe("Apostrophe");
 const FKey EKeys::Quote("Quote");
 
+const FKey EKeys::Asterix("Asterix");
+const FKey EKeys::Ampersand("Ampersand");
+const FKey EKeys::Caret("Caret");
+const FKey EKeys::Dollar("Dollar");
+const FKey EKeys::Exclamation("Exclamation");
+const FKey EKeys::Colon("Colon");
+
+const FKey EKeys::A_AccentGrave("A_AccentGrave");
+const FKey EKeys::E_AccentGrave("E_AccentGrave");
+const FKey EKeys::E_AccentAigu("E_AccentAigu");
+const FKey EKeys::C_Cedille("C_Cedille");
+
+
 // Setup platform specific keys
 #if PLATFORM_MAC
 const FKey EKeys::Platform_Delete = EKeys::BackSpace;
@@ -349,6 +366,7 @@ void EKeys::Initialize()
 	AddKey(FKeyDetails(EKeys::Semicolon, LOCTEXT("Semicolon", ";")));
 	AddKey(FKeyDetails(EKeys::Equals, LOCTEXT("Equals", "=")));
 	AddKey(FKeyDetails(EKeys::Comma, LOCTEXT("Comma", ",")));
+	AddKey(FKeyDetails(EKeys::Hyphen, LOCTEXT("Hyphen", "-")));
 	AddKey(FKeyDetails(EKeys::Underscore, LOCTEXT("Underscore", "_")));
 	AddKey(FKeyDetails(EKeys::Period, LOCTEXT("Period", ".")));
 	AddKey(FKeyDetails(EKeys::Slash, LOCTEXT("Slash", "/")));
@@ -356,7 +374,23 @@ void EKeys::Initialize()
 	AddKey(FKeyDetails(EKeys::LeftBracket, LOCTEXT("LeftBracket", "[")));
 	AddKey(FKeyDetails(EKeys::Backslash, LOCTEXT("Backslash", "\\")));
 	AddKey(FKeyDetails(EKeys::RightBracket, LOCTEXT("RightBracket", "]")));
-	AddKey(FKeyDetails(EKeys::Quote, LOCTEXT("Quote", "'")));
+	AddKey(FKeyDetails(EKeys::Apostrophe, LOCTEXT("Apostrophe", "'")));
+	AddKey(FKeyDetails(EKeys::Quote, LOCTEXT("Quote", "\"")));
+
+	AddKey(FKeyDetails(EKeys::LeftParantheses, LOCTEXT("LeftParantheses", "(")));
+	AddKey(FKeyDetails(EKeys::RightParantheses, LOCTEXT("RightParantheses", ")")));
+	AddKey(FKeyDetails(EKeys::Ampersand, LOCTEXT("Ampersand", "&")));
+	AddKey(FKeyDetails(EKeys::Asterix, LOCTEXT("Asterix", "*")));
+	AddKey(FKeyDetails(EKeys::Caret, LOCTEXT("Caret", "^")));
+	AddKey(FKeyDetails(EKeys::Dollar, LOCTEXT("Dollar", "$")));
+	AddKey(FKeyDetails(EKeys::Exclamation, LOCTEXT("Exclamation", "!")));
+	AddKey(FKeyDetails(EKeys::Colon, LOCTEXT("Colon", ":")));
+
+	AddKey(FKeyDetails(EKeys::A_AccentGrave, LOCTEXT("A_AccentGrave", "à")));
+	AddKey(FKeyDetails(EKeys::E_AccentGrave, LOCTEXT("E_AccentGrave", "é")));
+	AddKey(FKeyDetails(EKeys::E_AccentAigu, LOCTEXT("E_AccentAigu", "è")));
+	AddKey(FKeyDetails(EKeys::C_Cedille, LOCTEXT("C_Cedille", "ç")));
+
 
 	// Setup Gamepad keys
 	AddKey(FKeyDetails(EKeys::Gamepad_LeftX, LOCTEXT("Gamepad_LeftX", "Gamepad Left X"), FKeyDetails::GamepadKey | FKeyDetails::FloatAxis));
