Fix CIS issues by having extended characters directly in string
Don't use LOCTEXT for key names that are static across languages

diff --git a/Engine/Source/Runtime/InputCore/Private/InputCoreTypes.cpp b/Engine/Source/Runtime/InputCore/Private/InputCoreTypes.cpp
--- a/Engine/Source/Runtime/InputCore/Private/InputCoreTypes.cpp
+++ b/Engine/Source/Runtime/InputCore/Private/InputCoreTypes.cpp
@@ -363,33 +363,33 @@ void EKeys::Initialize()
 	AddKey(FKeyDetails(EKeys::LeftCommand, LOCTEXT("LeftCommand", "Left Cmd"), FKeyDetails::ModifierKey));
 	AddKey(FKeyDetails(EKeys::RightCommand, LOCTEXT("RightCommand", "Right Cmd"), FKeyDetails::ModifierKey));
 
-	AddKey(FKeyDetails(EKeys::Semicolon, LOCTEXT("Semicolon", ";")));
-	AddKey(FKeyDetails(EKeys::Equals, LOCTEXT("Equals", "=")));
-	AddKey(FKeyDetails(EKeys::Comma, LOCTEXT("Comma", ",")));
-	AddKey(FKeyDetails(EKeys::Hyphen, LOCTEXT("Hyphen", "-")));
-	AddKey(FKeyDetails(EKeys::Underscore, LOCTEXT("Underscore", "_")));
-	AddKey(FKeyDetails(EKeys::Period, LOCTEXT("Period", ".")));
-	AddKey(FKeyDetails(EKeys::Slash, LOCTEXT("Slash", "/")));
-	AddKey(FKeyDetails(EKeys::Tilde, LOCTEXT("Backtick", "`"))); // Yes this is not actually a tilde, it is a long, sad, and old story
-	AddKey(FKeyDetails(EKeys::LeftBracket, LOCTEXT("LeftBracket", "[")));
-	AddKey(FKeyDetails(EKeys::Backslash, LOCTEXT("Backslash", "\\")));
-	AddKey(FKeyDetails(EKeys::RightBracket, LOCTEXT("RightBracket", "]")));
-	AddKey(FKeyDetails(EKeys::Apostrophe, LOCTEXT("Apostrophe", "'")));
-	AddKey(FKeyDetails(EKeys::Quote, LOCTEXT("Quote", "\"")));
-
-	AddKey(FKeyDetails(EKeys::LeftParantheses, LOCTEXT("LeftParantheses", "(")));
-	AddKey(FKeyDetails(EKeys::RightParantheses, LOCTEXT("RightParantheses", ")")));
-	AddKey(FKeyDetails(EKeys::Ampersand, LOCTEXT("Ampersand", "&")));
-	AddKey(FKeyDetails(EKeys::Asterix, LOCTEXT("Asterix", "*")));
-	AddKey(FKeyDetails(EKeys::Caret, LOCTEXT("Caret", "^")));
-	AddKey(FKeyDetails(EKeys::Dollar, LOCTEXT("Dollar", "$")));
-	AddKey(FKeyDetails(EKeys::Exclamation, LOCTEXT("Exclamation", "!")));
-	AddKey(FKeyDetails(EKeys::Colon, LOCTEXT("Colon", ":")));
-
-	AddKey(FKeyDetails(EKeys::A_AccentGrave, LOCTEXT("A_AccentGrave", "à")));
-	AddKey(FKeyDetails(EKeys::E_AccentGrave, LOCTEXT("E_AccentGrave", "é")));
-	AddKey(FKeyDetails(EKeys::E_AccentAigu, LOCTEXT("E_AccentAigu", "è")));
-	AddKey(FKeyDetails(EKeys::C_Cedille, LOCTEXT("C_Cedille", "ç")));
+	AddKey(FKeyDetails(EKeys::Semicolon, FText::FromString(";")));
+	AddKey(FKeyDetails(EKeys::Equals, FText::FromString("=")));
+	AddKey(FKeyDetails(EKeys::Comma, FText::FromString(",")));
+	AddKey(FKeyDetails(EKeys::Hyphen, FText::FromString("-")));
+	AddKey(FKeyDetails(EKeys::Underscore, FText::FromString("_")));
+	AddKey(FKeyDetails(EKeys::Period, FText::FromString(".")));
+	AddKey(FKeyDetails(EKeys::Slash, FText::FromString("/")));
+	AddKey(FKeyDetails(EKeys::Tilde, FText::FromString("`"))); // Yes this is not actually a tilde, it is a long, sad, and old story
+	AddKey(FKeyDetails(EKeys::LeftBracket, FText::FromString("[")));
+	AddKey(FKeyDetails(EKeys::Backslash, FText::FromString("\\")));
+	AddKey(FKeyDetails(EKeys::RightBracket, FText::FromString("]")));
+	AddKey(FKeyDetails(EKeys::Apostrophe, FText::FromString("'")));
+	AddKey(FKeyDetails(EKeys::Quote, FText::FromString("\"")));
+
+	AddKey(FKeyDetails(EKeys::LeftParantheses, FText::FromString("(")));
+	AddKey(FKeyDetails(EKeys::RightParantheses, FText::FromString(")")));
+	AddKey(FKeyDetails(EKeys::Ampersand, FText::FromString("&")));
+	AddKey(FKeyDetails(EKeys::Asterix, FText::FromString("*")));
+	AddKey(FKeyDetails(EKeys::Caret, FText::FromString("^")));
+	AddKey(FKeyDetails(EKeys::Dollar, FText::FromString("$")));
+	AddKey(FKeyDetails(EKeys::Exclamation, FText::FromString("!")));
+	AddKey(FKeyDetails(EKeys::Colon, FText::FromString(":")));
+
+	AddKey(FKeyDetails(EKeys::A_AccentGrave, FText::FromString(FString::Chr(224))));
+	AddKey(FKeyDetails(EKeys::E_AccentGrave, FText::FromString(FString::Chr(232))));
+	AddKey(FKeyDetails(EKeys::E_AccentAigu, FText::FromString(FString::Chr(233))));
+	AddKey(FKeyDetails(EKeys::C_Cedille, FText::FromString(FString::Chr(231))));
 
 
 	// Setup Gamepad keys
