Tweak the display names of the gamepad key mappings

diff --git a/Engine/Source/Runtime/InputCore/Private/InputCoreTypes.cpp b/Engine/Source/Runtime/InputCore/Private/InputCoreTypes.cpp
--- a/Engine/Source/Runtime/InputCore/Private/InputCoreTypes.cpp
+++ b/Engine/Source/Runtime/InputCore/Private/InputCoreTypes.cpp
@@ -393,28 +393,26 @@ void EKeys::Initialize()
 
 
 	// Setup Gamepad keys
-	AddKey(FKeyDetails(EKeys::Gamepad_LeftX, LOCTEXT("Gamepad_LeftX", "Gamepad Left X"), FKeyDetails::GamepadKey | FKeyDetails::FloatAxis));
-	AddKey(FKeyDetails(EKeys::Gamepad_LeftY, LOCTEXT("Gamepad_LeftY", "Gamepad Left Y"), FKeyDetails::GamepadKey | FKeyDetails::FloatAxis));
-	AddKey(FKeyDetails(EKeys::Gamepad_RightX, LOCTEXT("Gamepad_RightX", "Gamepad Right X"), FKeyDetails::GamepadKey | FKeyDetails::FloatAxis));
-	AddKey(FKeyDetails(EKeys::Gamepad_RightY, LOCTEXT("Gamepad_RightY", "Gamepad Right Y"), FKeyDetails::GamepadKey | FKeyDetails::FloatAxis));
-
-	AddKey(FKeyDetails(EKeys::Gamepad_LeftThumbstick, LOCTEXT("Gamepad_LeftThumbstick", "Gamepad Left Thumbstick"), FKeyDetails::GamepadKey));
-	AddKey(FKeyDetails(EKeys::Gamepad_RightThumbstick, LOCTEXT("Gamepad_RightThumbstick", "Gamepad Right Thumbstick"), FKeyDetails::GamepadKey));
-	AddKey(FKeyDetails(EKeys::Gamepad_DPad_Up, LOCTEXT("Gamepad_DPad_Up", "Gamepad DPad Up"), FKeyDetails::GamepadKey));
-	AddKey(FKeyDetails(EKeys::Gamepad_DPad_Down, LOCTEXT("Gamepad_DPad_Down", "Gamepad DPad Down"), FKeyDetails::GamepadKey));
-	AddKey(FKeyDetails(EKeys::Gamepad_DPad_Right, LOCTEXT("Gamepad_DPad_Right", "Gamepad DPad Right"), FKeyDetails::GamepadKey));
-	AddKey(FKeyDetails(EKeys::Gamepad_DPad_Left, LOCTEXT("Gamepad_DPad_Left", "Gamepad DPad Left"), FKeyDetails::GamepadKey));
+	AddKey(FKeyDetails(EKeys::Gamepad_LeftX, LOCTEXT("Gamepad_LeftX", "Gamepad Left Thumbstick X-Axis"), FKeyDetails::GamepadKey | FKeyDetails::FloatAxis));
+	AddKey(FKeyDetails(EKeys::Gamepad_LeftY, LOCTEXT("Gamepad_LeftY", "Gamepad Left Thumbstick Y-Axis"), FKeyDetails::GamepadKey | FKeyDetails::FloatAxis));
+	AddKey(FKeyDetails(EKeys::Gamepad_RightX, LOCTEXT("Gamepad_RightX", "Gamepad Right Thumbstick X-Axis"), FKeyDetails::GamepadKey | FKeyDetails::FloatAxis));
+	AddKey(FKeyDetails(EKeys::Gamepad_RightY, LOCTEXT("Gamepad_RightY", "Gamepad Right Thumbstick Y-Axis"), FKeyDetails::GamepadKey | FKeyDetails::FloatAxis));
+
+	AddKey(FKeyDetails(EKeys::Gamepad_DPad_Up, LOCTEXT("Gamepad_DPad_Up", "Gamepad D-pad Up"), FKeyDetails::GamepadKey));
+	AddKey(FKeyDetails(EKeys::Gamepad_DPad_Down, LOCTEXT("Gamepad_DPad_Down", "Gamepad D-pad Down"), FKeyDetails::GamepadKey));
+	AddKey(FKeyDetails(EKeys::Gamepad_DPad_Right, LOCTEXT("Gamepad_DPad_Right", "Gamepad D-pad Right"), FKeyDetails::GamepadKey));
+	AddKey(FKeyDetails(EKeys::Gamepad_DPad_Left, LOCTEXT("Gamepad_DPad_Left", "Gamepad D-pad Left"), FKeyDetails::GamepadKey));
 
 	// Virtual key codes used for input axis button press/release emulation
-	AddKey(FKeyDetails(EKeys::Gamepad_LeftStick_Up, LOCTEXT("Gamepad_LeftStick_Up", "Gamepad Left Stick Up"), FKeyDetails::GamepadKey));
-	AddKey(FKeyDetails(EKeys::Gamepad_LeftStick_Down, LOCTEXT("Gamepad_LeftStick_Down", "Gamepad Left Stick Down"), FKeyDetails::GamepadKey));
-	AddKey(FKeyDetails(EKeys::Gamepad_LeftStick_Right, LOCTEXT("Gamepad_LeftStick_Right", "Gamepad Left Stick Right"), FKeyDetails::GamepadKey));
-	AddKey(FKeyDetails(EKeys::Gamepad_LeftStick_Left, LOCTEXT("Gamepad_LeftStick_Left", "Gamepad Left Stick Left"), FKeyDetails::GamepadKey));
+	AddKey(FKeyDetails(EKeys::Gamepad_LeftStick_Up, LOCTEXT("Gamepad_LeftStick_Up", "Gamepad Left Thumbstick Up"), FKeyDetails::GamepadKey));
+	AddKey(FKeyDetails(EKeys::Gamepad_LeftStick_Down, LOCTEXT("Gamepad_LeftStick_Down", "Gamepad Left Thumbstick Down"), FKeyDetails::GamepadKey));
+	AddKey(FKeyDetails(EKeys::Gamepad_LeftStick_Right, LOCTEXT("Gamepad_LeftStick_Right", "Gamepad Left Thumbstick Right"), FKeyDetails::GamepadKey));
+	AddKey(FKeyDetails(EKeys::Gamepad_LeftStick_Left, LOCTEXT("Gamepad_LeftStick_Left", "Gamepad Left Thumbstick Left"), FKeyDetails::GamepadKey));
 
-	AddKey(FKeyDetails(EKeys::Gamepad_RightStick_Up, LOCTEXT("Gamepad_RightStick_Up", "Gamepad Right Stick Up"), FKeyDetails::GamepadKey));
-	AddKey(FKeyDetails(EKeys::Gamepad_RightStick_Down, LOCTEXT("Gamepad_RightStick_Down", "Gamepad Right Stick Down"), FKeyDetails::GamepadKey));
-	AddKey(FKeyDetails(EKeys::Gamepad_RightStick_Right, LOCTEXT("Gamepad_RightStick_Right", "Gamepad Right Stick Right"), FKeyDetails::GamepadKey));
-	AddKey(FKeyDetails(EKeys::Gamepad_RightStick_Left, LOCTEXT("Gamepad_RightStick_Left", "Gamepad Right Stick Left"), FKeyDetails::GamepadKey));
+	AddKey(FKeyDetails(EKeys::Gamepad_RightStick_Up, LOCTEXT("Gamepad_RightStick_Up", "Gamepad Right Thumbstick Up"), FKeyDetails::GamepadKey));
+	AddKey(FKeyDetails(EKeys::Gamepad_RightStick_Down, LOCTEXT("Gamepad_RightStick_Down", "Gamepad Right Thumbstick Down"), FKeyDetails::GamepadKey));
+	AddKey(FKeyDetails(EKeys::Gamepad_RightStick_Right, LOCTEXT("Gamepad_RightStick_Right", "Gamepad Right Thumbstick Right"), FKeyDetails::GamepadKey));
+	AddKey(FKeyDetails(EKeys::Gamepad_RightStick_Left, LOCTEXT("Gamepad_RightStick_Left", "Gamepad Right Thumbstick Left"), FKeyDetails::GamepadKey));
 
 	AddKey(FKeyDetails(EKeys::Gamepad_Special_Left, TAttribute<FText>::Create(TAttribute<FText>::FGetter::CreateStatic(&EKeys::GetGamepadDisplayName, EKeys::Gamepad_Special_Left)), FKeyDetails::GamepadKey));
 	AddKey(FKeyDetails(EKeys::Gamepad_Special_Right, TAttribute<FText>::Create(TAttribute<FText>::FGetter::CreateStatic(&EKeys::GetGamepadDisplayName, EKeys::Gamepad_Special_Right)), FKeyDetails::GamepadKey));
@@ -431,6 +429,9 @@ void EKeys::Initialize()
 	AddKey(FKeyDetails(EKeys::Gamepad_LeftTrigger, TAttribute<FText>::Create(TAttribute<FText>::FGetter::CreateStatic(&EKeys::GetGamepadDisplayName, EKeys::Gamepad_LeftTrigger)), FKeyDetails::GamepadKey));
 	AddKey(FKeyDetails(EKeys::Gamepad_RightTrigger, TAttribute<FText>::Create(TAttribute<FText>::FGetter::CreateStatic(&EKeys::GetGamepadDisplayName, EKeys::Gamepad_RightTrigger)), FKeyDetails::GamepadKey));
 
+	AddKey(FKeyDetails(EKeys::Gamepad_LeftThumbstick, TAttribute<FText>::Create(TAttribute<FText>::FGetter::CreateStatic(&EKeys::GetGamepadDisplayName, EKeys::Gamepad_LeftThumbstick)), FKeyDetails::GamepadKey));
+	AddKey(FKeyDetails(EKeys::Gamepad_RightThumbstick, TAttribute<FText>::Create(TAttribute<FText>::FGetter::CreateStatic(&EKeys::GetGamepadDisplayName, EKeys::Gamepad_RightThumbstick)), FKeyDetails::GamepadKey));
+
 	// Vector axes (FVector, not float)
 	AddKey(FKeyDetails(EKeys::Tilt, LOCTEXT("Tilt", "Tilt"), FKeyDetails::VectorAxis));
 	AddKey(FKeyDetails(EKeys::RotationRate, LOCTEXT("RotationRate", "Rotation Rate"), FKeyDetails::VectorAxis));
@@ -551,6 +552,14 @@ FText EKeys::GetGamepadDisplayName(const FKey Key)
 		{
 			return LOCTEXT("PS4_Gamepad_RightTriggerAxis", "Gamepad R2 Axis");
 		}
+		else if (Key == EKeys::Gamepad_LeftThumbstick)
+		{
+			return LOCTEXT("PS4_Gamepad_LeftThumbstick", "Gamepad L3");
+		}
+		else if (Key == EKeys::Gamepad_RightThumbstick)
+		{
+			return LOCTEXT("PS4_Gamepad_RightThumbstick", "Gamepad R3");
+		}
 		break;
 
 	case EConsoleForGamepadLabels::XBoxOne:
@@ -602,6 +611,14 @@ FText EKeys::GetGamepadDisplayName(const FKey Key)
 		{
 			return LOCTEXT("Gamepad_RightTriggerAxis", "Gamepad Right Trigger Axis");
 		}
+		else if (Key == EKeys::Gamepad_LeftThumbstick)
+		{
+			return LOCTEXT("Gamepad_LeftThumbstick", "Gamepad Left Thumbstick Button");
+		}
+		else if (Key == EKeys::Gamepad_RightThumbstick)
+		{
+			return LOCTEXT("Gamepad_RightThumbstick", "Gamepad Right Thumbstick Button");
+		}
 		break;
 
 	default:
@@ -653,6 +670,14 @@ FText EKeys::GetGamepadDisplayName(const FKey Key)
 		{
 			return LOCTEXT("Gamepad_RightTriggerAxis", "Gamepad Right Trigger Axis");
 		}
+		else if (Key == EKeys::Gamepad_LeftThumbstick)
+		{
+			return LOCTEXT("Gamepad_LeftThumbstick", "Gamepad Left Thumbstick Button");
+		}
+		else if (Key == EKeys::Gamepad_RightThumbstick)
+		{
+			return LOCTEXT("Gamepad_RightThumbstick", "Gamepad Right Thumbstick Button");
+		}
 		break;
 	}
 
