BP - The FKey structure now copies using the assignment operator so that the KeyDetails shared ptr is properly copied to the new copy.

diff --git a/Engine/Source/Runtime/InputCore/Classes/InputCoreTypes.h b/Engine/Source/Runtime/InputCore/Classes/InputCoreTypes.h
--- a/Engine/Source/Runtime/InputCore/Classes/InputCoreTypes.h
+++ b/Engine/Source/Runtime/InputCore/Classes/InputCoreTypes.h
@@ -71,6 +71,7 @@ struct TStructOpsTypeTraits<FKey> : public TStructOpsTypeTraitsBase
 		WithExportTextItem = true,
 		WithImportTextItem = true,
 		WithPostSerialize = true,
+		WithCopy = true,		// Necessary so that TSharedPtr<FKeyDetails> Data is copied around
 	};
 };
 
