UE-8738 Use full words for keys that are hard to read in UI

diff --git a/Engine/Source/Runtime/InputCore/Private/InputCoreTypes.cpp b/Engine/Source/Runtime/InputCore/Private/InputCoreTypes.cpp
--- a/Engine/Source/Runtime/InputCore/Private/InputCoreTypes.cpp
+++ b/Engine/Source/Runtime/InputCore/Private/InputCoreTypes.cpp
@@ -363,28 +363,28 @@ void EKeys::Initialize()
 	AddKey(FKeyDetails(EKeys::LeftCommand, LOCTEXT("LeftCommand", "Left Cmd"), FKeyDetails::ModifierKey));
 	AddKey(FKeyDetails(EKeys::RightCommand, LOCTEXT("RightCommand", "Right Cmd"), FKeyDetails::ModifierKey));
 
-	AddKey(FKeyDetails(EKeys::Semicolon, FText::FromString(";")));
+	AddKey(FKeyDetails(EKeys::Semicolon, LOCTEXT("Semicolon", "Semicolon")));
 	AddKey(FKeyDetails(EKeys::Equals, FText::FromString("=")));
-	AddKey(FKeyDetails(EKeys::Comma, FText::FromString(",")));
-	AddKey(FKeyDetails(EKeys::Hyphen, FText::FromString("-")));
-	AddKey(FKeyDetails(EKeys::Underscore, FText::FromString("_")));
-	AddKey(FKeyDetails(EKeys::Period, FText::FromString(".")));
+	AddKey(FKeyDetails(EKeys::Comma, LOCTEXT("Comma", "Comma")));
+	AddKey(FKeyDetails(EKeys::Hyphen, LOCTEXT("Hyphen", "Hyphen")));
+	AddKey(FKeyDetails(EKeys::Underscore, LOCTEXT("Underscore", "Underscore")));
+	AddKey(FKeyDetails(EKeys::Period, LOCTEXT("Period", "Period")));
 	AddKey(FKeyDetails(EKeys::Slash, FText::FromString("/")));
 	AddKey(FKeyDetails(EKeys::Tilde, FText::FromString("`"))); // Yes this is not actually a tilde, it is a long, sad, and old story
 	AddKey(FKeyDetails(EKeys::LeftBracket, FText::FromString("[")));
 	AddKey(FKeyDetails(EKeys::Backslash, FText::FromString("\\")));
 	AddKey(FKeyDetails(EKeys::RightBracket, FText::FromString("]")));
-	AddKey(FKeyDetails(EKeys::Apostrophe, FText::FromString("'")));
+	AddKey(FKeyDetails(EKeys::Apostrophe, LOCTEXT("Apostrophe", "Apostrophe")));
 	AddKey(FKeyDetails(EKeys::Quote, FText::FromString("\"")));
 
 	AddKey(FKeyDetails(EKeys::LeftParantheses, FText::FromString("(")));
 	AddKey(FKeyDetails(EKeys::RightParantheses, FText::FromString(")")));
 	AddKey(FKeyDetails(EKeys::Ampersand, FText::FromString("&")));
-	AddKey(FKeyDetails(EKeys::Asterix, FText::FromString("*")));
+	AddKey(FKeyDetails(EKeys::Asterix, LOCTEXT("Asterix", "Asterix")));
 	AddKey(FKeyDetails(EKeys::Caret, FText::FromString("^")));
 	AddKey(FKeyDetails(EKeys::Dollar, FText::FromString("$")));
 	AddKey(FKeyDetails(EKeys::Exclamation, FText::FromString("!")));
-	AddKey(FKeyDetails(EKeys::Colon, FText::FromString(":")));
+	AddKey(FKeyDetails(EKeys::Colon, LOCTEXT("Colon", "Colon")));
 
 	AddKey(FKeyDetails(EKeys::A_AccentGrave, FText::FromString(FString::Chr(224))));
 	AddKey(FKeyDetails(EKeys::E_AccentGrave, FText::FromString(FString::Chr(232))));
