- Exposed Steam controller buttons to blueprint, input mapping, etc.

diff --git a/Engine/Source/Runtime/InputCore/Classes/InputCoreTypes.h b/Engine/Source/Runtime/InputCore/Classes/InputCoreTypes.h
--- a/Engine/Source/Runtime/InputCore/Classes/InputCoreTypes.h
+++ b/Engine/Source/Runtime/InputCore/Classes/InputCoreTypes.h
@@ -371,7 +371,7 @@ struct INPUTCORE_API EKeys
 	// PS4-specific
 	static const FKey PS4_Special;
 
-	// Steam Controller Specifit;
+	// Steam Controller Specific
 	static const FKey Steam_Touch_0;
 	static const FKey Steam_Touch_1;
 	static const FKey Steam_Touch_2;
diff --git a/Engine/Source/Runtime/InputCore/Private/InputCoreTypes.cpp b/Engine/Source/Runtime/InputCore/Private/InputCoreTypes.cpp
--- a/Engine/Source/Runtime/InputCore/Private/InputCoreTypes.cpp
+++ b/Engine/Source/Runtime/InputCore/Private/InputCoreTypes.cpp
@@ -222,6 +222,14 @@ const FKey EKeys::Gesture_Flick("Gesture_Flick");
 // PS4-specific
 const FKey EKeys::PS4_Special("PS4_Special");
 
+// Steam Controller Specific
+const FKey EKeys::Steam_Touch_0("Steam_Touch_0");
+const FKey EKeys::Steam_Touch_1("Steam_Touch_1");
+const FKey EKeys::Steam_Touch_2("Steam_Touch_2");
+const FKey EKeys::Steam_Touch_3("Steam_Touch_3");
+const FKey EKeys::Steam_Back_Left("Steam_Back_Left");
+const FKey EKeys::Steam_Back_Right("Steam_Back_Right");
+
 // Xbox One global speech commands
 const FKey EKeys::Global_Menu("Global_Menu");
 const FKey EKeys::Global_View("Global_View");
@@ -461,6 +469,14 @@ void EKeys::Initialize()
 	// PS4-specific
 	AddKey(FKeyDetails(EKeys::PS4_Special, LOCTEXT("PS4_Special", "PS4_Special"), FKeyDetails::NotBlueprintBindableKey));
 
+	// Steam Controller specific
+	AddKey(FKeyDetails(EKeys::Steam_Touch_0, LOCTEXT("Steam_Touch_0", "Steam Touch 0"), FKeyDetails::GamepadKey));
+	AddKey(FKeyDetails(EKeys::Steam_Touch_1, LOCTEXT("Steam_Touch_1", "Steam Touch 1"), FKeyDetails::GamepadKey));
+	AddKey(FKeyDetails(EKeys::Steam_Touch_2, LOCTEXT("Steam_Touch_2", "Steam Touch 2"), FKeyDetails::GamepadKey));
+	AddKey(FKeyDetails(EKeys::Steam_Touch_3, LOCTEXT("Steam_Touch_3", "Steam Touch 3"), FKeyDetails::GamepadKey));
+	AddKey(FKeyDetails(EKeys::Steam_Back_Left, LOCTEXT("Steam_Back_Left", "Steam Back Left"), FKeyDetails::GamepadKey));
+	AddKey(FKeyDetails(EKeys::Steam_Back_Right, LOCTEXT("Steam_Back_Right", "Steam Back Right"), FKeyDetails::GamepadKey));
+
 	// Xbox One global speech commands
 	AddKey(FKeyDetails(EKeys::Global_Menu, LOCTEXT("Global_Menu", "Global Menu"), FKeyDetails::GamepadKey));
 	AddKey(FKeyDetails(EKeys::Global_View, LOCTEXT("Global_View", "Global View"), FKeyDetails::GamepadKey));
diff --git a/Engine/Source/Runtime/Slate/Private/Framework/Application/SlateApplication.cpp b/Engine/Source/Runtime/Slate/Private/Framework/Application/SlateApplication.cpp
--- a/Engine/Source/Runtime/Slate/Private/Framework/Application/SlateApplication.cpp
+++ b/Engine/Source/Runtime/Slate/Private/Framework/Application/SlateApplication.cpp
@@ -4663,6 +4663,13 @@ FKey TranslateControllerButtonToKey( EControllerButtons::Type Button )
 	case EControllerButtons::RightStickLeft: Key = EKeys::Gamepad_RightStick_Left; break;
 	case EControllerButtons::RightStickRight: Key = EKeys::Gamepad_RightStick_Right; break;
 
+	case EControllerButtons::Touch0: Key = EKeys::Steam_Touch_0; break;
+	case EControllerButtons::Touch1: Key = EKeys::Steam_Touch_1; break;
+	case EControllerButtons::Touch2: Key = EKeys::Steam_Touch_2; break;
+	case EControllerButtons::Touch3: Key = EKeys::Steam_Touch_3; break;
+	case EControllerButtons::BackLeft: Key = EKeys::Steam_Back_Left; break;
+	case EControllerButtons::BackRight: Key = EKeys::Steam_Back_Right; break;
+
 	case EControllerButtons::GlobalMenu: Key = EKeys::Global_Menu; break;
 	case EControllerButtons::GlobalView: Key = EKeys::Global_View; break;
 	case EControllerButtons::GlobalPause: Key = EKeys::Global_Pause; break;
