Remove unsupported gestures from EKeys
Made keys that are not blueprint bindable not selectable via the key selector dropdowns
Made supported gesture keys blueprint bindable
Added mechanism to organize keys in to more categories than just gamepad, keyboard, and mouse
Moved motion, gesture, android, steam, and xbox one keys in to their own menu categories

diff --git a/Engine/Source/Editor/BlueprintGraph/Private/K2Node_GetInputAxisKeyValue.cpp b/Engine/Source/Editor/BlueprintGraph/Private/K2Node_GetInputAxisKeyValue.cpp
--- a/Engine/Source/Editor/BlueprintGraph/Private/K2Node_GetInputAxisKeyValue.cpp
+++ b/Engine/Source/Editor/BlueprintGraph/Private/K2Node_GetInputAxisKeyValue.cpp
@@ -111,21 +111,9 @@ void UK2Node_GetInputAxisKeyValue::RegisterDynamicBinding(UDynamicBlueprintBindi
 
 FName UK2Node_GetInputAxisKeyValue::GetPaletteIcon(FLinearColor& OutColor) const
 {
-	if (InputAxisKey.IsMouseButton())
-	{
-		return TEXT("GraphEditor.MouseEvent_16x");
-	}
-	else if (InputAxisKey.IsGamepadKey())
-	{
-		return TEXT("GraphEditor.PadEvent_16x");
-	}
-	else
-	{
-		return TEXT("GraphEditor.KeyEvent_16x");
-	}
+	return EKeys::GetMenuCategoryPaletteIcon(InputAxisKey.GetMenuCategory());
 }
 
-
 void UK2Node_GetInputAxisKeyValue::GetMenuActions(FBlueprintActionDatabaseRegistrar& ActionRegistrar) const
 {
 	TArray<FKey> AllKeys;
@@ -169,40 +157,18 @@ void UK2Node_GetInputAxisKeyValue::GetMenuActions(FBlueprintActionDatabaseRegist
 
 FText UK2Node_GetInputAxisKeyValue::GetMenuCategory() const
 {
-	enum EAxisKeyCategory
-	{
-		GamepadKeyCategory,
-		MouseButtonCategory,
-		KeyValueCategory,
-		AxisKeyCategory_MAX,
-	};
-	static FNodeTextCache CachedCategories[AxisKeyCategory_MAX];
-
-	FText SubCategory;
-	EAxisKeyCategory CategoryIndex = AxisKeyCategory_MAX;
+	static TMap<FName, FNodeTextCache> CachedCategories;
 
-	if (InputAxisKey.IsGamepadKey())
-	{
-		SubCategory = LOCTEXT("GamepadSubCategory", "Gamepad Values");
-		CategoryIndex = GamepadKeyCategory;
-	}
-	else if (InputAxisKey.IsMouseButton())
-	{
-		SubCategory = LOCTEXT("MouseSubCategory", "Mouse Values");
-		CategoryIndex = MouseButtonCategory;
-	}
-	else
-	{
-		SubCategory = LOCTEXT("KeySubCategory", "Key Values");
-		CategoryIndex = KeyValueCategory;
-	}
+	const FName KeyCategory = InputAxisKey.GetMenuCategory();
+	const FText SubCategoryDisplayName = FText::Format(LOCTEXT("EventsCategory", "{0} Values"), EKeys::GetMenuCategoryDisplayName(KeyCategory));
+	FNodeTextCache& NodeTextCache = CachedCategories.FindOrAdd(KeyCategory);
 
-	if (CachedCategories[CategoryIndex].IsOutOfDate(this))
+	if (NodeTextCache.IsOutOfDate(this))
 	{
 		// FText::Format() is slow, so we cache this to save on performance
-		CachedCategories[CategoryIndex].SetCachedText(FEditorCategoryUtils::BuildCategoryString(FCommonEditorCategory::Input, SubCategory), this);
+		NodeTextCache.SetCachedText(FEditorCategoryUtils::BuildCategoryString(FCommonEditorCategory::Input, SubCategoryDisplayName), this);
 	}
-	return CachedCategories[CategoryIndex];
+	return NodeTextCache;
 }
 
 FBlueprintNodeSignature UK2Node_GetInputAxisKeyValue::GetSignature() const
diff --git a/Engine/Source/Editor/BlueprintGraph/Private/K2Node_InputAxisKeyEvent.cpp b/Engine/Source/Editor/BlueprintGraph/Private/K2Node_InputAxisKeyEvent.cpp
--- a/Engine/Source/Editor/BlueprintGraph/Private/K2Node_InputAxisKeyEvent.cpp
+++ b/Engine/Source/Editor/BlueprintGraph/Private/K2Node_InputAxisKeyEvent.cpp
@@ -81,18 +81,7 @@ UClass* UK2Node_InputAxisKeyEvent::GetDynamicBindingClass() const
 
 FName UK2Node_InputAxisKeyEvent::GetPaletteIcon(FLinearColor& OutColor) const
 {
-	if (AxisKey.IsMouseButton())
-	{
-		return TEXT("GraphEditor.MouseEvent_16x");
-	}
-	else if (AxisKey.IsGamepadKey())
-	{
-		return TEXT("GraphEditor.PadEvent_16x");
-	}
-	else
-	{
-		return TEXT("GraphEditor.KeyEvent_16x");
-	}
+	return EKeys::GetMenuCategoryPaletteIcon(AxisKey.GetMenuCategory());
 }
 
 void UK2Node_InputAxisKeyEvent::RegisterDynamicBinding(UDynamicBlueprintBinding* BindingObject) const
@@ -171,40 +160,18 @@ void UK2Node_InputAxisKeyEvent::GetMenuActions(FBlueprintActionDatabaseRegistrar
 
 FText UK2Node_InputAxisKeyEvent::GetMenuCategory() const
 {
-	enum EAxisKeyCategory
-	{
-		GamepadKeyCategory,
-		MouseButtonCategory,
-		KeyEventCategory,
-		AxisKeyCategory_MAX,
-	};
-	static FNodeTextCache CachedCategories[AxisKeyCategory_MAX];
+	static TMap<FName, FNodeTextCache> CachedCategories;
 
-	FText SubCategory;
-	EAxisKeyCategory CategoryIndex = AxisKeyCategory_MAX;
-
-	if (AxisKey.IsGamepadKey())
-	{
-		SubCategory = LOCTEXT("GamepadCategory", "Gamepad Events");
-		CategoryIndex = GamepadKeyCategory;
-	}
-	else if (AxisKey.IsMouseButton())
-	{
-		SubCategory = LOCTEXT("MouseCategory", "Mouse Events");
-		CategoryIndex = MouseButtonCategory;
-	}
-	else
-	{
-		SubCategory = LOCTEXT("KeyEventsCategory", "Key Events");
-		CategoryIndex = KeyEventCategory;
-	}
+	const FName KeyCategory = AxisKey.GetMenuCategory();
+	const FText SubCategoryDisplayName = FText::Format(LOCTEXT("EventsCategory", "{0} Events"), EKeys::GetMenuCategoryDisplayName(KeyCategory));
+	FNodeTextCache& NodeTextCache = CachedCategories.FindOrAdd(KeyCategory);
 
-	if (CachedCategories[CategoryIndex].IsOutOfDate(this))
+	if (NodeTextCache.IsOutOfDate(this))
 	{
 		// FText::Format() is slow, so we cache this to save on performance
-		CachedCategories[CategoryIndex].SetCachedText(FEditorCategoryUtils::BuildCategoryString(FCommonEditorCategory::Input, SubCategory), this);
+		NodeTextCache.SetCachedText(FEditorCategoryUtils::BuildCategoryString(FCommonEditorCategory::Input, SubCategoryDisplayName), this);
 	}
-	return CachedCategories[CategoryIndex];
+	return NodeTextCache;
 }
 
 FBlueprintNodeSignature UK2Node_InputAxisKeyEvent::GetSignature() const
diff --git a/Engine/Source/Editor/BlueprintGraph/Private/K2Node_InputKey.cpp b/Engine/Source/Editor/BlueprintGraph/Private/K2Node_InputKey.cpp
--- a/Engine/Source/Editor/BlueprintGraph/Private/K2Node_InputKey.cpp
+++ b/Engine/Source/Editor/BlueprintGraph/Private/K2Node_InputKey.cpp
@@ -193,18 +193,7 @@ FText UK2Node_InputKey::GetTooltipText() const
 
 FName UK2Node_InputKey::GetPaletteIcon(FLinearColor& OutColor) const
 {
-	if (InputKey.IsMouseButton())
-	{
-		return TEXT("GraphEditor.MouseEvent_16x");
-	}
-	else if (InputKey.IsGamepadKey())
-	{
-		return TEXT("GraphEditor.PadEvent_16x");
-	}
-	else
-	{
-		return TEXT("GraphEditor.KeyEvent_16x");
-	}
+	return EKeys::GetMenuCategoryPaletteIcon(InputKey.GetMenuCategory());
 }
 
 bool UK2Node_InputKey::IsCompatibleWithGraph(UEdGraph const* Graph) const
@@ -334,40 +323,18 @@ void UK2Node_InputKey::GetMenuActions(FBlueprintActionDatabaseRegistrar& ActionR
 
 FText UK2Node_InputKey::GetMenuCategory() const
 {
-	enum EAxisKeyCategory
-	{
-		GamepadKeyCategory,
-		MouseButtonCategory,
-		KeyEventCategory,
-		AxisKeyCategory_MAX,
-	};
-	static FNodeTextCache CachedCategories[AxisKeyCategory_MAX];
+	static TMap<FName, FNodeTextCache> CachedCategories;
 
-	FText SubCategory;
-	EAxisKeyCategory CategoryIndex = AxisKeyCategory_MAX;
-
-	if (InputKey.IsGamepadKey())
-	{
-		SubCategory = LOCTEXT("GamepadCategory", "Gamepad Events");
-		CategoryIndex = GamepadKeyCategory;
-	}
-	else if (InputKey.IsMouseButton())
-	{
-		SubCategory = LOCTEXT("MouseCategory", "Mouse Events");
-		CategoryIndex = MouseButtonCategory;
-	}
-	else
-	{
-		SubCategory = LOCTEXT("KeyEventsCategory", "Key Events");
-		CategoryIndex = KeyEventCategory;
-	}
+	const FName KeyCategory = InputKey.GetMenuCategory();
+	const FText SubCategoryDisplayName = FText::Format(LOCTEXT("EventsCategory", "{0} Events"), EKeys::GetMenuCategoryDisplayName(KeyCategory));
+	FNodeTextCache& NodeTextCache = CachedCategories.FindOrAdd(KeyCategory);
 
-	if (CachedCategories[CategoryIndex].IsOutOfDate(this))
+	if (NodeTextCache.IsOutOfDate(this))
 	{
 		// FText::Format() is slow, so we cache this to save on performance
-		CachedCategories[CategoryIndex].SetCachedText(FEditorCategoryUtils::BuildCategoryString(FCommonEditorCategory::Input, SubCategory), this);
+		NodeTextCache.SetCachedText(FEditorCategoryUtils::BuildCategoryString(FCommonEditorCategory::Input, SubCategoryDisplayName), this);
 	}
-	return CachedCategories[CategoryIndex];
+	return NodeTextCache;
 }
 
 FBlueprintNodeSignature UK2Node_InputKey::GetSignature() const
diff --git a/Engine/Source/Editor/UnrealEd/Private/SKeySelector.cpp b/Engine/Source/Editor/UnrealEd/Private/SKeySelector.cpp
--- a/Engine/Source/Editor/UnrealEd/Private/SKeySelector.cpp
+++ b/Engine/Source/Editor/UnrealEd/Private/SKeySelector.cpp
@@ -78,9 +78,12 @@ void SKeySelector::Construct(const FArguments& InArgs)
 	OnKeyChanged = InArgs._OnKeyChanged;
 	CurrentKey = InArgs._CurrentKey;
 
-	auto PadCategory = *new (KeyTreeRoot) FKeyTreeItem(MakeShareable(new FKeyTreeInfo(LOCTEXT("PadCategory", "Gamepad"), nullptr)));
-	auto KeyCategory = *new (KeyTreeRoot) FKeyTreeItem(MakeShareable(new FKeyTreeInfo(LOCTEXT("KeyCategory", "Keyboard"), nullptr)));
-	auto MouseCategory = *new (KeyTreeRoot) FKeyTreeItem(MakeShareable(new FKeyTreeInfo(LOCTEXT("MouseCategory", "Mouse"), nullptr)));
+	TMap<FName, FKeyTreeItem> TreeRootsForCatgories;
+
+	// Ensure that Gamepad, Keyboard, and Mouse will appear at the top of the list, other categories will dynamically get added as the keys are encountered
+	TreeRootsForCatgories.Add(EKeys::NAME_GamepadCategory, *new (KeyTreeRoot) FKeyTreeItem(MakeShareable(new FKeyTreeInfo(EKeys::GetMenuCategoryDisplayName(EKeys::NAME_GamepadCategory), nullptr))));
+	TreeRootsForCatgories.Add(EKeys::NAME_KeyboardCategory, *new (KeyTreeRoot) FKeyTreeItem(MakeShareable(new FKeyTreeInfo(EKeys::GetMenuCategoryDisplayName(EKeys::NAME_KeyboardCategory), nullptr))));
+	TreeRootsForCatgories.Add(EKeys::NAME_MouseCategory, *new (KeyTreeRoot) FKeyTreeItem(MakeShareable(new FKeyTreeInfo(EKeys::GetMenuCategoryDisplayName(EKeys::NAME_MouseCategory), nullptr))));
 
 	TArray<FKey> AllKeys;
 	EKeys::GetAllKeys(AllKeys);
@@ -89,19 +92,14 @@ void SKeySelector::Construct(const FArguments& InArgs)
 	{
 		if (!InArgs._FilterBlueprintBindable || Key.IsBindableInBlueprints())
 		{
-			auto NewKeyItem = MakeShareable(new FKeyTreeInfo(FText(), MakeShareable(new FKey(Key))));
-			if (Key.IsMouseButton())
-			{
-				MouseCategory->Children.Add(NewKeyItem);
-			}
-			else if (Key.IsGamepadKey())
+			const FName KeyMenuCategory = Key.GetMenuCategory();
+			FKeyTreeItem* KeyCategory = TreeRootsForCatgories.Find(KeyMenuCategory);
+			if (KeyCategory == nullptr)
 			{
-				PadCategory->Children.Add(NewKeyItem);
-			}
-			else
-			{
-				KeyCategory->Children.Add(NewKeyItem);
+				KeyCategory = new (KeyTreeRoot) FKeyTreeItem(MakeShareable(new FKeyTreeInfo(EKeys::GetMenuCategoryDisplayName(KeyMenuCategory), nullptr)));
+				TreeRootsForCatgories.Add(KeyMenuCategory, *KeyCategory);
 			}
+			(*KeyCategory)->Children.Add(MakeShareable(new FKeyTreeInfo(FText(), MakeShareable(new FKey(Key)))));
 		}
 	}
 
@@ -401,22 +399,7 @@ bool SKeySelector::GetChildrenMatchingSearch(const TArray<FString>& InSearchToke
 
 const FSlateBrush* SKeySelector::GetIconFromKey(FKey Key) const
 {
-	static const FName NAME_MouseEvent(TEXT("GraphEditor.MouseEvent_16x"));
-	static const FName NAME_PadEvent(TEXT("GraphEditor.PadEvent_16x"));
-	static const FName NAME_KeyEvent(TEXT("GraphEditor.KeyEvent_16x"));
-
-	if (Key.IsMouseButton())
-	{
-		return FEditorStyle::GetBrush(NAME_MouseEvent);
-	}
-	else if (Key.IsGamepadKey())
-	{
-		return FEditorStyle::GetBrush(NAME_PadEvent);
-	}
-	else
-	{
-		return FEditorStyle::GetBrush(NAME_KeyEvent);
-	}
+	return FEditorStyle::GetBrush(EKeys::GetMenuCategoryPaletteIcon(Key.GetMenuCategory()));
 }
 
 #undef LOCTEXT_NAMESPACE
diff --git a/Engine/Source/Editor/UnrealEd/Public/SKeySelector.h b/Engine/Source/Editor/UnrealEd/Public/SKeySelector.h
--- a/Engine/Source/Editor/UnrealEd/Public/SKeySelector.h
+++ b/Engine/Source/Editor/UnrealEd/Public/SKeySelector.h
@@ -22,7 +22,7 @@ public:
 		, _TreeViewWidth(300.f)
 		, _TreeViewHeight(400.f)
 		, _Font( FEditorStyle::GetFontStyle( TEXT("NormalFont") ) )
-		, _FilterBlueprintBindable( false )
+		, _FilterBlueprintBindable( true )
 		, _HasMultipleValues( false )
 		{}
 		SLATE_ATTRIBUTE( FKey, CurrentKey )
diff --git a/Engine/Source/Runtime/Engine/Classes/GameFramework/PlayerInput.h b/Engine/Source/Runtime/Engine/Classes/GameFramework/PlayerInput.h
--- a/Engine/Source/Runtime/Engine/Classes/GameFramework/PlayerInput.h
+++ b/Engine/Source/Runtime/Engine/Classes/GameFramework/PlayerInput.h
@@ -13,8 +13,6 @@
 #include "KeyState.h"
 #include "PlayerInput.generated.h"
 
-ENGINE_API DECLARE_LOG_CATEGORY_EXTERN(LogInput, Log, All);
-
 /** Struct containing mappings for legacy method of binding keys to exec commands. */
 USTRUCT()
 struct FKeyBind
diff --git a/Engine/Source/Runtime/Engine/Private/GestureRecognizer.cpp b/Engine/Source/Runtime/Engine/Private/GestureRecognizer.cpp
--- a/Engine/Source/Runtime/Engine/Private/GestureRecognizer.cpp
+++ b/Engine/Source/Runtime/Engine/Private/GestureRecognizer.cpp
@@ -56,14 +56,6 @@ void FGestureRecognizer::DetectGestures(const FVector (&Touches)[EKeys::NUM_TOUC
 			}
 		}
 
-		// cancel any 2 finger gestures
-		if (PreviousTouchCount >= 2 && TouchCount < 2)
-		{
-			HandleGesture(PlayerInput, EKeys::Gesture_TwoFingerSwipeLeftRight, false, true);
-			HandleGesture(PlayerInput, EKeys::Gesture_TwoFingerSwipeUpDown, false, true);
-			HandleGesture(PlayerInput, EKeys::Gesture_Pinch, false, true);
-		}
-
 		if (PreviousTouchCount == 0 && TouchCount == 1)
 		{
 			// initialize the flick
@@ -89,13 +81,6 @@ void FGestureRecognizer::DetectGestures(const FVector (&Touches)[EKeys::NUM_TOUC
 				HandleGesture(PlayerInput, EKeys::Gesture_Flick, true, true);
 			}
 		}
-
-		if (PreviousTouchCount >= 1 && TouchCount == 0)
-		{
-			// cancel other 1 finger gestures
-			HandleGesture(PlayerInput, EKeys::Gesture_SwipeLeftRight, false, true);
-			HandleGesture(PlayerInput, EKeys::Gesture_SwipeUpDown, false, true);
-		}
 	}	
 
 	// remember for next frame
diff --git a/Engine/Source/Runtime/Engine/Private/UserInterface/PlayerInput.cpp b/Engine/Source/Runtime/Engine/Private/UserInterface/PlayerInput.cpp
--- a/Engine/Source/Runtime/Engine/Private/UserInterface/PlayerInput.cpp
+++ b/Engine/Source/Runtime/Engine/Private/UserInterface/PlayerInput.cpp
@@ -8,8 +8,6 @@
 #include "GameFramework/PlayerInput.h"
 #include "GameFramework/InputSettings.h"
 
-DEFINE_LOG_CATEGORY(LogInput);
-
 bool bExecutingBindCommand = false;
 
 /** for debug rendering */
diff --git a/Engine/Source/Runtime/InputCore/Classes/InputCoreTypes.h b/Engine/Source/Runtime/InputCore/Classes/InputCoreTypes.h
--- a/Engine/Source/Runtime/InputCore/Classes/InputCoreTypes.h
+++ b/Engine/Source/Runtime/InputCore/Classes/InputCoreTypes.h
@@ -4,6 +4,8 @@
 
 #include "InputCoreTypes.generated.h"
 
+INPUTCORE_API DECLARE_LOG_CATEGORY_EXTERN(LogInput, Log, All);
+
 USTRUCT(BlueprintType)
 struct INPUTCORE_API FKey
 {
@@ -38,6 +40,7 @@ struct INPUTCORE_API FKey
 	FText GetDisplayName() const;
 	FString ToString() const;
 	FName GetFName() const;
+	FName GetMenuCategory() const;
 
 	bool SerializeFromMismatchedTag(struct FPropertyTag const& Tag, FArchive& Ar);
 	bool ExportTextItem(FString& ValueStr, FKey const& DefaultValue, UObject* Parent, int32 PortFlags, UObject* ExportRootScope) const;
@@ -91,25 +94,7 @@ struct INPUTCORE_API FKeyDetails
 		NoFlags                 = 0,
 	};
 
-	FKeyDetails(const FKey InKey, const TAttribute<FText>& InDisplayName, const uint8 InKeyFlags = 0)
-		: Key(InKey)
-		, DisplayName(InDisplayName)
-		, bIsModifierKey((InKeyFlags & EKeyFlags::ModifierKey) != 0)
-		, bIsGamepadKey((InKeyFlags & EKeyFlags::GamepadKey) != 0)
-		, bIsMouseButton((InKeyFlags & EKeyFlags::MouseButton) != 0)
-		, bIsBindableInBlueprints((~InKeyFlags & EKeyFlags::NotBlueprintBindableKey) != 0)
-		, AxisType(EInputAxisType::None)
-	{
-		if ((InKeyFlags & EKeyFlags::FloatAxis) != 0)
-		{
-			ensure((InKeyFlags & EKeyFlags::VectorAxis) == 0);
-			AxisType = EInputAxisType::Float;
-		}
-		else if ((InKeyFlags & EKeyFlags::VectorAxis) != 0)
-		{
-			AxisType = EInputAxisType::Vector;
-		}
-	}
+	FKeyDetails(const FKey InKey, const TAttribute<FText>& InDisplayName, const uint8 InKeyFlags = 0, const FName InMenuCategory = NAME_None);
 
 	bool IsModifierKey() const { return bIsModifierKey != 0; }
 	bool IsGamepadKey() const { return bIsGamepadKey != 0; }
@@ -117,6 +102,7 @@ struct INPUTCORE_API FKeyDetails
 	bool IsFloatAxis() const { return AxisType == EInputAxisType::Float; }
 	bool IsVectorAxis() const { return AxisType == EInputAxisType::Vector; }
 	bool IsBindableInBlueprints() const { return bIsBindableInBlueprints != 0; }
+	FName GetMenuCategory() const { return MenuCategory; }
 	FText GetDisplayName() const;
 	const FKey& GetKey() const { return Key; }
 
@@ -133,6 +119,8 @@ private:
 	
 	TAttribute<FText> DisplayName;
 
+	FName MenuCategory;
+
 	int32 bIsModifierKey:1;
 	int32 bIsGamepadKey:1;
 	int32 bIsMouseButton:1;
@@ -361,10 +349,6 @@ struct INPUTCORE_API EKeys
 	static const FKey Acceleration;
 
 	// Gestures
-	static const FKey Gesture_SwipeLeftRight;
-	static const FKey Gesture_SwipeUpDown;
-	static const FKey Gesture_TwoFingerSwipeLeftRight;
-	static const FKey Gesture_TwoFingerSwipeUpDown;
 	static const FKey Gesture_Pinch;
 	static const FKey Gesture_Flick;
 
@@ -399,6 +383,10 @@ struct INPUTCORE_API EKeys
 
 	static EConsoleForGamepadLabels::Type ConsoleForGamepadLabels;
 
+	static const FName NAME_KeyboardCategory;
+	static const FName NAME_GamepadCategory;
+	static const FName NAME_MouseCategory;
+
 	static void Initialize();
 	static void AddKey(const FKeyDetails& KeyDetails);
 	static void GetAllKeys(TArray<FKey>& OutKeys);
@@ -414,9 +402,19 @@ struct INPUTCORE_API EKeys
 	// Function that provides remapping for some gamepad keys in display windows
 	static FText GetGamepadDisplayName(const FKey Key);
 
+	static void AddMenuCategoryDisplayInfo(const FName CategoryName, const FText DisplayName, const FName PaletteIcon);
+	static FText GetMenuCategoryDisplayName(const FName CategoryName);
+	static FName GetMenuCategoryPaletteIcon(const FName CategoryName);
 private:
 
+	struct FCategoryDisplayInfo
+	{
+		FText DisplayName;
+		FName PaletteIcon;
+	};
+
 	static TMap<FKey, TSharedPtr<FKeyDetails> > InputKeys;
+	static TMap<FName, FCategoryDisplayInfo> MenuCategoryDisplayInfo;
 	static bool bInitialized;
 
 };
diff --git a/Engine/Source/Runtime/InputCore/Private/InputCoreTypes.cpp b/Engine/Source/Runtime/InputCore/Private/InputCoreTypes.cpp
--- a/Engine/Source/Runtime/InputCore/Private/InputCoreTypes.cpp
+++ b/Engine/Source/Runtime/InputCore/Private/InputCoreTypes.cpp
@@ -4,6 +4,8 @@
 #include "PropertyTag.h"
 #include "InputCoreTypes.h"
 
+DEFINE_LOG_CATEGORY(LogInput);
+
 #define LOCTEXT_NAMESPACE "InputKeys"
 
 const FKey EKeys::MouseX("MouseX");
@@ -212,10 +214,6 @@ const FKey EKeys::TouchKeys[NUM_TOUCH_KEYS] =
 };
 
 // Gestures
-const FKey EKeys::Gesture_SwipeLeftRight("Gesture_SwipeLeftRight");
-const FKey EKeys::Gesture_SwipeUpDown("Gesture_SwipeUpDown");
-const FKey EKeys::Gesture_TwoFingerSwipeLeftRight("Gesture_TwoFingerSwipeLeftRight");
-const FKey EKeys::Gesture_TwoFingerSwipeUpDown("Gesture_TwoFingerSwipeUpDown");
 const FKey EKeys::Gesture_Pinch("Gesture_Pinch");
 const FKey EKeys::Gesture_Flick("Gesture_Flick");
 
@@ -245,14 +243,61 @@ const FKey EKeys::Android_Menu("Android_Menu");
 
 const FKey EKeys::Invalid(NAME_None);
 
+const FName EKeys::NAME_GamepadCategory("Gamepad");
+const FName EKeys::NAME_MouseCategory("Mouse");
+const FName EKeys::NAME_KeyboardCategory("Key");
+
 bool EKeys::bInitialized = false;
 TMap<FKey, TSharedPtr<FKeyDetails> > EKeys::InputKeys;
+TMap<FName, EKeys::FCategoryDisplayInfo> EKeys::MenuCategoryDisplayInfo;
+
+FKeyDetails::FKeyDetails(const FKey InKey, const TAttribute<FText>& InDisplayName, const uint8 InKeyFlags, const FName InMenuCategory)
+	: Key(InKey)
+	, DisplayName(InDisplayName)
+	, bIsModifierKey((InKeyFlags & EKeyFlags::ModifierKey) != 0)
+	, bIsGamepadKey((InKeyFlags & EKeyFlags::GamepadKey) != 0)
+	, bIsMouseButton((InKeyFlags & EKeyFlags::MouseButton) != 0)
+	, bIsBindableInBlueprints((~InKeyFlags & EKeyFlags::NotBlueprintBindableKey) != 0)
+	, AxisType(EInputAxisType::None)
+	, MenuCategory(InMenuCategory)
+{
+	if ((InKeyFlags & EKeyFlags::FloatAxis) != 0)
+	{
+		ensure((InKeyFlags & EKeyFlags::VectorAxis) == 0);
+		AxisType = EInputAxisType::Float;
+	}
+	else if ((InKeyFlags & EKeyFlags::VectorAxis) != 0)
+	{
+		AxisType = EInputAxisType::Vector;
+	}
+
+	// Set up default menu categories
+	if (MenuCategory.IsNone())
+	{
+		if (IsGamepadKey())
+		{
+			MenuCategory = EKeys::NAME_GamepadCategory;
+		}
+		else if (IsMouseButton())
+		{
+			MenuCategory = EKeys::NAME_MouseCategory;
+		}
+		else
+		{
+			MenuCategory = EKeys::NAME_KeyboardCategory;
+		}
+	}
+}
 
 void EKeys::Initialize()
 {
 	if (bInitialized) return;
 	bInitialized = true;
 
+	AddMenuCategoryDisplayInfo(NAME_GamepadCategory, LOCTEXT("GamepadSubCategory", "Gamepad"), TEXT("GraphEditor.PadEvent_16x"));
+	AddMenuCategoryDisplayInfo(NAME_MouseCategory, LOCTEXT("MouseSubCategory", "Mouse"), TEXT("GraphEditor.MouseEvent_16x"));
+	AddMenuCategoryDisplayInfo(NAME_KeyboardCategory, LOCTEXT("KeyboardSubCategory", "Keyboard"), TEXT("GraphEditor.KeyEvent_16x"));
+
 	AddKey(FKeyDetails(EKeys::MouseX, LOCTEXT("MouseX", "Mouse X"), FKeyDetails::FloatAxis | FKeyDetails::MouseButton));
 	AddKey(FKeyDetails(EKeys::MouseY, LOCTEXT("MouseY", "Mouse Y"), FKeyDetails::FloatAxis | FKeyDetails::MouseButton));
 	AddKey(FKeyDetails(EKeys::MouseScrollUp, LOCTEXT("MouseScrollUp", "Mouse Wheel Up"), FKeyDetails::MouseButton));
@@ -440,11 +485,12 @@ void EKeys::Initialize()
 	AddKey(FKeyDetails(EKeys::Gamepad_LeftThumbstick, TAttribute<FText>::Create(TAttribute<FText>::FGetter::CreateStatic(&EKeys::GetGamepadDisplayName, EKeys::Gamepad_LeftThumbstick)), FKeyDetails::GamepadKey));
 	AddKey(FKeyDetails(EKeys::Gamepad_RightThumbstick, TAttribute<FText>::Create(TAttribute<FText>::FGetter::CreateStatic(&EKeys::GetGamepadDisplayName, EKeys::Gamepad_RightThumbstick)), FKeyDetails::GamepadKey));
 
-	// Vector axes (FVector, not float)
-	AddKey(FKeyDetails(EKeys::Tilt, LOCTEXT("Tilt", "Tilt"), FKeyDetails::VectorAxis));
-	AddKey(FKeyDetails(EKeys::RotationRate, LOCTEXT("RotationRate", "Rotation Rate"), FKeyDetails::VectorAxis));
-	AddKey(FKeyDetails(EKeys::Gravity, LOCTEXT("Gravity", "Gravity"), FKeyDetails::VectorAxis));
-	AddKey(FKeyDetails(EKeys::Acceleration, LOCTEXT("Acceleration", "Acceleration"), FKeyDetails::VectorAxis));
+	AddMenuCategoryDisplayInfo("Motion", LOCTEXT("MotionSubCateogry", "Motion"), TEXT("GraphEditor.KeyEvent_16x"));
+
+	AddKey(FKeyDetails(EKeys::Tilt, LOCTEXT("Tilt", "Tilt"), FKeyDetails::VectorAxis, "Motion"));
+	AddKey(FKeyDetails(EKeys::RotationRate, LOCTEXT("RotationRate", "Rotation Rate"), FKeyDetails::VectorAxis, "Motion"));
+	AddKey(FKeyDetails(EKeys::Gravity, LOCTEXT("Gravity", "Gravity"), FKeyDetails::VectorAxis, "Motion"));
+	AddKey(FKeyDetails(EKeys::Acceleration, LOCTEXT("Acceleration", "Acceleration"), FKeyDetails::VectorAxis, "Motion"));
 
 	// Fingers
 	AddKey(FKeyDetails(EKeys::TouchKeys[ETouchIndex::Touch1], LOCTEXT("Touch1", "Touch 1"), FKeyDetails::NotBlueprintBindableKey));
@@ -459,36 +505,40 @@ void EKeys::Initialize()
 	AddKey(FKeyDetails(EKeys::TouchKeys[ETouchIndex::Touch10], LOCTEXT("Touch10", "Touch 10"), FKeyDetails::NotBlueprintBindableKey));
 
 	// Gestures
-	AddKey(FKeyDetails(EKeys::Gesture_SwipeLeftRight, LOCTEXT("Gesture_SwipeLeftRight", "Swipe Left To Right"), FKeyDetails::NotBlueprintBindableKey));
-	AddKey(FKeyDetails(EKeys::Gesture_SwipeUpDown, LOCTEXT("Gesture_SwipeUpDown", "Swipe Up To Down"), FKeyDetails::NotBlueprintBindableKey));
-	AddKey(FKeyDetails(EKeys::Gesture_TwoFingerSwipeLeftRight, LOCTEXT("Gesture_TwoFingerSwipeLeftRight", "Two Finger Swipe Left To Right"), FKeyDetails::NotBlueprintBindableKey));
-	AddKey(FKeyDetails(EKeys::Gesture_TwoFingerSwipeUpDown, LOCTEXT("Gesture_TwoFingerSwipeUpDown", "Two Finger Swipe Up To Down"), FKeyDetails::NotBlueprintBindableKey));
-	AddKey(FKeyDetails(EKeys::Gesture_Pinch, LOCTEXT("Gesture_Pinch", "Pinch"), FKeyDetails::NotBlueprintBindableKey));
-	AddKey(FKeyDetails(EKeys::Gesture_Flick, LOCTEXT("Gesture_Flick", "Flick"), FKeyDetails::NotBlueprintBindableKey));
+	AddMenuCategoryDisplayInfo("Gesture", LOCTEXT("GestureSubCateogry", "Gesture"), TEXT("GraphEditor.KeyEvent_16x"));
+
+	AddKey(FKeyDetails(EKeys::Gesture_Pinch, LOCTEXT("Gesture_Pinch", "Pinch"), 0, "Gesture"));
+	AddKey(FKeyDetails(EKeys::Gesture_Flick, LOCTEXT("Gesture_Flick", "Flick"), 0, "Gesture"));
 
 	// PS4-specific
 	AddKey(FKeyDetails(EKeys::PS4_Special, LOCTEXT("PS4_Special", "PS4_Special"), FKeyDetails::NotBlueprintBindableKey));
 
 	// Steam Controller specific
-	AddKey(FKeyDetails(EKeys::Steam_Touch_0, LOCTEXT("Steam_Touch_0", "Steam Touch 0"), FKeyDetails::GamepadKey));
-	AddKey(FKeyDetails(EKeys::Steam_Touch_1, LOCTEXT("Steam_Touch_1", "Steam Touch 1"), FKeyDetails::GamepadKey));
-	AddKey(FKeyDetails(EKeys::Steam_Touch_2, LOCTEXT("Steam_Touch_2", "Steam Touch 2"), FKeyDetails::GamepadKey));
-	AddKey(FKeyDetails(EKeys::Steam_Touch_3, LOCTEXT("Steam_Touch_3", "Steam Touch 3"), FKeyDetails::GamepadKey));
-	AddKey(FKeyDetails(EKeys::Steam_Back_Left, LOCTEXT("Steam_Back_Left", "Steam Back Left"), FKeyDetails::GamepadKey));
-	AddKey(FKeyDetails(EKeys::Steam_Back_Right, LOCTEXT("Steam_Back_Right", "Steam Back Right"), FKeyDetails::GamepadKey));
+	AddMenuCategoryDisplayInfo("Steam", LOCTEXT("SteamSubCateogry", "Steam"), TEXT("GraphEditor.PadEvent_16x"));
+
+	AddKey(FKeyDetails(EKeys::Steam_Touch_0, LOCTEXT("Steam_Touch_0", "Steam Touch 0"), FKeyDetails::GamepadKey, "Steam"));
+	AddKey(FKeyDetails(EKeys::Steam_Touch_1, LOCTEXT("Steam_Touch_1", "Steam Touch 1"), FKeyDetails::GamepadKey, "Steam"));
+	AddKey(FKeyDetails(EKeys::Steam_Touch_2, LOCTEXT("Steam_Touch_2", "Steam Touch 2"), FKeyDetails::GamepadKey, "Steam"));
+	AddKey(FKeyDetails(EKeys::Steam_Touch_3, LOCTEXT("Steam_Touch_3", "Steam Touch 3"), FKeyDetails::GamepadKey, "Steam"));
+	AddKey(FKeyDetails(EKeys::Steam_Back_Left, LOCTEXT("Steam_Back_Left", "Steam Back Left"), FKeyDetails::GamepadKey, "Steam"));
+	AddKey(FKeyDetails(EKeys::Steam_Back_Right, LOCTEXT("Steam_Back_Right", "Steam Back Right"), FKeyDetails::GamepadKey, "Steam"));
 
 	// Xbox One global speech commands
-	AddKey(FKeyDetails(EKeys::Global_Menu, LOCTEXT("Global_Menu", "Global Menu"), FKeyDetails::GamepadKey));
-	AddKey(FKeyDetails(EKeys::Global_View, LOCTEXT("Global_View", "Global View"), FKeyDetails::GamepadKey));
-	AddKey(FKeyDetails(EKeys::Global_Pause, LOCTEXT("Global_Pause", "Global Pause"), FKeyDetails::GamepadKey));
-	AddKey(FKeyDetails(EKeys::Global_Play, LOCTEXT("Global_Play", "Global Play"), FKeyDetails::GamepadKey));
-	AddKey(FKeyDetails(EKeys::Global_Back, LOCTEXT("Global_Back", "Global Back"), FKeyDetails::GamepadKey));
+	AddMenuCategoryDisplayInfo("XBoxOne", LOCTEXT("XBoxOneSubCateogry", "XBox One"), TEXT("GraphEditor.PadEvent_16x"));
+
+	AddKey(FKeyDetails(EKeys::Global_Menu, LOCTEXT("Global_Menu", "Global Menu"), FKeyDetails::GamepadKey, "XBoxOne"));
+	AddKey(FKeyDetails(EKeys::Global_View, LOCTEXT("Global_View", "Global View"), FKeyDetails::GamepadKey, "XBoxOne"));
+	AddKey(FKeyDetails(EKeys::Global_Pause, LOCTEXT("Global_Pause", "Global Pause"), FKeyDetails::GamepadKey, "XBoxOne"));
+	AddKey(FKeyDetails(EKeys::Global_Play, LOCTEXT("Global_Play", "Global Play"), FKeyDetails::GamepadKey, "XBoxOne"));
+	AddKey(FKeyDetails(EKeys::Global_Back, LOCTEXT("Global_Back", "Global Back"), FKeyDetails::GamepadKey, "XBoxOne"));
 
 	// Android-specific
-	AddKey(FKeyDetails(EKeys::Android_Back, LOCTEXT("Android_Back", "Android Back"), FKeyDetails::GamepadKey));
-	AddKey(FKeyDetails(EKeys::Android_Volume_Up, LOCTEXT("Android_Volume_Up", "Android Volume Up"), FKeyDetails::GamepadKey));
-	AddKey(FKeyDetails(EKeys::Android_Volume_Down, LOCTEXT("Android_Volume_Down", "Android Volume Down"), FKeyDetails::GamepadKey));
-	AddKey(FKeyDetails(EKeys::Android_Menu, LOCTEXT("Android_Menu", "Android Menu"), FKeyDetails::GamepadKey));
+	AddMenuCategoryDisplayInfo("Android", LOCTEXT("AndroidSubCateogry", "Android"), TEXT("GraphEditor.KeyEvent_16x"));
+
+	AddKey(FKeyDetails(EKeys::Android_Back, LOCTEXT("Android_Back", "Android Back"), FKeyDetails::GamepadKey, "Android"));
+	AddKey(FKeyDetails(EKeys::Android_Volume_Up, LOCTEXT("Android_Volume_Up", "Android Volume Up"), FKeyDetails::GamepadKey, "Android"));
+	AddKey(FKeyDetails(EKeys::Android_Volume_Down, LOCTEXT("Android_Volume_Down", "Android Volume Down"), FKeyDetails::GamepadKey, "Android"));
+	AddKey(FKeyDetails(EKeys::Android_Menu, LOCTEXT("Android_Menu", "Android Menu"), FKeyDetails::GamepadKey, "Android"));
 
 	// Initialize the input key manager.  This will cause any additional OEM keys to get added
 	FInputKeyManager::Get();
@@ -517,6 +567,33 @@ TSharedPtr<FKeyDetails> EKeys::GetKeyDetails(const FKey Key)
 	return *KeyDetails;
 }
 
+void EKeys::AddMenuCategoryDisplayInfo(const FName CategoryName, const FText DisplayName, const FName PaletteIcon)
+{
+	if (MenuCategoryDisplayInfo.Contains(CategoryName))
+	{
+		UE_LOG(LogInput, Warning, TEXT("Category %s already has menu display info that is being replaced."), *CategoryName.ToString());
+	}
+	MenuCategoryDisplayInfo.FindOrAdd(CategoryName) = { DisplayName, PaletteIcon };
+}
+
+FText EKeys::GetMenuCategoryDisplayName(const FName CategoryName)
+{
+	if (FCategoryDisplayInfo* DisplayInfo = MenuCategoryDisplayInfo.Find(CategoryName))
+	{
+		return DisplayInfo->DisplayName;
+	}
+	return MenuCategoryDisplayInfo[NAME_KeyboardCategory].DisplayName;
+}
+
+FName EKeys::GetMenuCategoryPaletteIcon(const FName CategoryName)
+{
+	if (FCategoryDisplayInfo* DisplayInfo = MenuCategoryDisplayInfo.Find(CategoryName))
+	{
+		return DisplayInfo->PaletteIcon;
+	}
+	return MenuCategoryDisplayInfo[NAME_KeyboardCategory].PaletteIcon;
+}
+
 EConsoleForGamepadLabels::Type EKeys::ConsoleForGamepadLabels = EConsoleForGamepadLabels::None;
 
 FText EKeys::GetGamepadDisplayName(const FKey Key)
@@ -772,6 +849,12 @@ FText FKey::GetDisplayName() const
 	return (KeyDetails.IsValid() ? KeyDetails->GetDisplayName() : FText::FromName(KeyName));
 }
 
+FName FKey::GetMenuCategory() const
+{
+	ConditionalLookupKeyDetails();
+	return (KeyDetails.IsValid() ? KeyDetails->GetMenuCategory() : EKeys::NAME_KeyboardCategory);
+}
+
 void FKey::ConditionalLookupKeyDetails() const
 {
 	if (!KeyDetails.IsValid())
