Fix XBoxOne compile error

diff --git a/Engine/Source/Runtime/InputCore/Private/InputCoreTypes.cpp b/Engine/Source/Runtime/InputCore/Private/InputCoreTypes.cpp
--- a/Engine/Source/Runtime/InputCore/Private/InputCoreTypes.cpp
+++ b/Engine/Source/Runtime/InputCore/Private/InputCoreTypes.cpp
@@ -573,7 +573,10 @@ void EKeys::AddMenuCategoryDisplayInfo(const FName CategoryName, const FText Dis
 	{
 		UE_LOG(LogInput, Warning, TEXT("Category %s already has menu display info that is being replaced."), *CategoryName.ToString());
 	}
-	MenuCategoryDisplayInfo.FindOrAdd(CategoryName) = { DisplayName, PaletteIcon };
+	FCategoryDisplayInfo DisplayInfo;
+	DisplayInfo.DisplayName = DisplayName;
+	DisplayInfo.PaletteIcon = PaletteIcon;
+	MenuCategoryDisplayInfo.FindOrAdd(CategoryName) = DisplayInfo;
 }
 
 FText EKeys::GetMenuCategoryDisplayName(const FName CategoryName)
