Avoid static initialization issues

diff --git a/Engine/Source/Runtime/Core/Private/GenericPlatform/GenericApplication.cpp b/Engine/Source/Runtime/Core/Private/GenericPlatform/GenericApplication.cpp
--- a/Engine/Source/Runtime/Core/Private/GenericPlatform/GenericApplication.cpp
+++ b/Engine/Source/Runtime/Core/Private/GenericPlatform/GenericApplication.cpp
@@ -6,6 +6,7 @@
 
 const FGamepadKeyNames::Type FGamepadKeyNames::Invalid(NAME_None);
 
+// Ensure that the GamepadKeyNames match those in InputCoreTypes.cpp
 const FGamepadKeyNames::Type FGamepadKeyNames::LeftAnalogX("Gamepad_LeftX");
 const FGamepadKeyNames::Type FGamepadKeyNames::LeftAnalogY("Gamepad_LeftY");
 const FGamepadKeyNames::Type FGamepadKeyNames::RightAnalogX("Gamepad_RightX");
diff --git a/Engine/Source/Runtime/InputCore/Private/InputCoreTypes.cpp b/Engine/Source/Runtime/InputCore/Private/InputCoreTypes.cpp
--- a/Engine/Source/Runtime/InputCore/Private/InputCoreTypes.cpp
+++ b/Engine/Source/Runtime/InputCore/Private/InputCoreTypes.cpp
@@ -157,40 +157,41 @@ const FKey EKeys::Platform_Delete = EKeys::BackSpace;
 const FKey EKeys::Platform_Delete = EKeys::Delete;
 #endif
 
-const FKey EKeys::Gamepad_LeftX(FGamepadKeyNames::LeftAnalogX);
-const FKey EKeys::Gamepad_LeftY(FGamepadKeyNames::LeftAnalogY);
-const FKey EKeys::Gamepad_RightX(FGamepadKeyNames::RightAnalogX);
-const FKey EKeys::Gamepad_RightY(FGamepadKeyNames::RightAnalogY);
-const FKey EKeys::Gamepad_LeftTriggerAxis(FGamepadKeyNames::LeftTriggerAnalog);
-const FKey EKeys::Gamepad_RightTriggerAxis(FGamepadKeyNames::RightTriggerAnalog);
-
-const FKey EKeys::Gamepad_LeftThumbstick(FGamepadKeyNames::LeftThumb);
-const FKey EKeys::Gamepad_RightThumbstick(FGamepadKeyNames::RightThumb);
-const FKey EKeys::Gamepad_Special_Left(FGamepadKeyNames::SpecialLeft);
-const FKey EKeys::Gamepad_Special_Right(FGamepadKeyNames::SpecialRight);
-const FKey EKeys::Gamepad_FaceButton_Bottom(FGamepadKeyNames::FaceButtonBottom);
-const FKey EKeys::Gamepad_FaceButton_Right(FGamepadKeyNames::FaceButtonRight);
-const FKey EKeys::Gamepad_FaceButton_Left(FGamepadKeyNames::FaceButtonLeft);
-const FKey EKeys::Gamepad_FaceButton_Top(FGamepadKeyNames::FaceButtonTop);
-const FKey EKeys::Gamepad_LeftShoulder(FGamepadKeyNames::LeftShoulder);
-const FKey EKeys::Gamepad_RightShoulder(FGamepadKeyNames::RightShoulder);
-const FKey EKeys::Gamepad_LeftTrigger(FGamepadKeyNames::LeftTriggerThreshold);
-const FKey EKeys::Gamepad_RightTrigger(FGamepadKeyNames::RightTriggerThreshold);
-const FKey EKeys::Gamepad_DPad_Up(FGamepadKeyNames::DPadUp);
-const FKey EKeys::Gamepad_DPad_Down(FGamepadKeyNames::DPadDown);
-const FKey EKeys::Gamepad_DPad_Right(FGamepadKeyNames::DPadRight);
-const FKey EKeys::Gamepad_DPad_Left(FGamepadKeyNames::DPadLeft);
+// Ensure that the Gamepad_ names match those in GenericApplication.cpp
+const FKey EKeys::Gamepad_LeftX("Gamepad_LeftX");
+const FKey EKeys::Gamepad_LeftY("Gamepad_LeftY");
+const FKey EKeys::Gamepad_RightX("Gamepad_RightX");
+const FKey EKeys::Gamepad_RightY("Gamepad_RightY");
+const FKey EKeys::Gamepad_LeftTriggerAxis("Gamepad_LeftTriggerAxis");
+const FKey EKeys::Gamepad_RightTriggerAxis("Gamepad_RightTriggerAxis");
+
+const FKey EKeys::Gamepad_LeftThumbstick("Gamepad_LeftThumbstick");
+const FKey EKeys::Gamepad_RightThumbstick("Gamepad_RightThumbstick");
+const FKey EKeys::Gamepad_Special_Left("Gamepad_Special_Left");
+const FKey EKeys::Gamepad_Special_Right("Gamepad_Special_Right");
+const FKey EKeys::Gamepad_FaceButton_Bottom("Gamepad_FaceButton_Bottom");
+const FKey EKeys::Gamepad_FaceButton_Right("Gamepad_FaceButton_Right");
+const FKey EKeys::Gamepad_FaceButton_Left("Gamepad_FaceButton_Left");
+const FKey EKeys::Gamepad_FaceButton_Top("Gamepad_FaceButton_Top");
+const FKey EKeys::Gamepad_LeftShoulder("Gamepad_LeftShoulder");
+const FKey EKeys::Gamepad_RightShoulder("Gamepad_RightShoulder");
+const FKey EKeys::Gamepad_LeftTrigger("Gamepad_LeftTrigger");
+const FKey EKeys::Gamepad_RightTrigger("Gamepad_RightTrigger");
+const FKey EKeys::Gamepad_DPad_Up("Gamepad_DPad_Up");
+const FKey EKeys::Gamepad_DPad_Down("Gamepad_DPad_Down");
+const FKey EKeys::Gamepad_DPad_Right("Gamepad_DPad_Right");
+const FKey EKeys::Gamepad_DPad_Left("Gamepad_DPad_Left");
 
 // Virtual key codes used for input axis button press/release emulation
-const FKey EKeys::Gamepad_LeftStick_Up(FGamepadKeyNames::LeftStickUp);
-const FKey EKeys::Gamepad_LeftStick_Down(FGamepadKeyNames::LeftStickDown);
-const FKey EKeys::Gamepad_LeftStick_Right(FGamepadKeyNames::LeftStickRight);
-const FKey EKeys::Gamepad_LeftStick_Left(FGamepadKeyNames::LeftStickLeft);
-
-const FKey EKeys::Gamepad_RightStick_Up(FGamepadKeyNames::RightStickUp);
-const FKey EKeys::Gamepad_RightStick_Down(FGamepadKeyNames::RightStickDown);
-const FKey EKeys::Gamepad_RightStick_Right(FGamepadKeyNames::RightStickRight);
-const FKey EKeys::Gamepad_RightStick_Left(FGamepadKeyNames::RightStickLeft);
+const FKey EKeys::Gamepad_LeftStick_Up("Gamepad_LeftStick_Up");
+const FKey EKeys::Gamepad_LeftStick_Down("Gamepad_LeftStick_Down");
+const FKey EKeys::Gamepad_LeftStick_Right("Gamepad_LeftStick_Right");
+const FKey EKeys::Gamepad_LeftStick_Left("Gamepad_LeftStick_Left");
+
+const FKey EKeys::Gamepad_RightStick_Up("Gamepad_RightStick_Up");
+const FKey EKeys::Gamepad_RightStick_Down("Gamepad_RightStick_Down");
+const FKey EKeys::Gamepad_RightStick_Right("Gamepad_RightStick_Right");
+const FKey EKeys::Gamepad_RightStick_Left("Gamepad_RightStick_Left");
 
 // const FKey EKeys::Vector axes (FVector("Vector axes (FVector"); not float)
 const FKey EKeys::Tilt("Tilt");
