#motioncontroller Added generic motion controller key abstraction, updated SteamVRController to use new abstraction, as well as fixing issues with d-pad pressing

diff --git a/Engine/Plugins/Runtime/Steam/SteamVR/Source/SteamVRController/Private/SteamVRController.cpp b/Engine/Plugins/Runtime/Steam/SteamVR/Source/SteamVRController/Private/SteamVRController.cpp
--- a/Engine/Plugins/Runtime/Steam/SteamVR/Source/SteamVRController/Private/SteamVRController.cpp
+++ b/Engine/Plugins/Runtime/Steam/SteamVR/Source/SteamVRController/Private/SteamVRController.cpp
@@ -8,6 +8,7 @@ DEFINE_LOG_CATEGORY_STATIC(LogSteamVRController, Log, All);
 /** Controller axis mappings. @todo steamvr: should enumerate rather than hard code */
 #define TOUCHPAD_AXIS	0
 #define TRIGGER_AXIS	1
+#define DOT_45DEG		0.7071f
 
 //
 // Gamepad thresholds
@@ -18,8 +19,6 @@ namespace SteamVRControllerKeyNames
 {
 	const FGamepadKeyNames::Type Touch0("Steam_Touch_0");
 	const FGamepadKeyNames::Type Touch1("Steam_Touch_1");
-	const FGamepadKeyNames::Type BackLeft("Steam_Back_Left");
-	const FGamepadKeyNames::Type BackRight("Steam_Back_Right");
 }
 
 class FSteamVRController : public IInputDevice
@@ -97,26 +96,26 @@ public:
 		ButtonRepeatDelay = 0.1f;
 
 		Buttons[ (int32)EControllerHand::Left ][ ESteamVRControllerButton::System ] = FGamepadKeyNames::SpecialLeft;
-		Buttons[ (int32)EControllerHand::Left ][ ESteamVRControllerButton::ApplicationMenu ] = FGamepadKeyNames::LeftShoulder;
-		Buttons[ (int32)EControllerHand::Left ][ ESteamVRControllerButton::TouchPadPress ] = FGamepadKeyNames::LeftThumb;
+		Buttons[ (int32)EControllerHand::Left ][ ESteamVRControllerButton::ApplicationMenu ] = FGamepadKeyNames::MotionController_Left_Shoulder;
+		Buttons[ (int32)EControllerHand::Left ][ ESteamVRControllerButton::TouchPadPress ] = FGamepadKeyNames::MotionController_Left_Thumbstick;
 		Buttons[ (int32)EControllerHand::Left ][ ESteamVRControllerButton::TouchPadTouch ] = SteamVRControllerKeyNames::Touch0;
-		Buttons[ (int32)EControllerHand::Left ][ ESteamVRControllerButton::TriggerPress ] = FGamepadKeyNames::LeftTriggerThreshold;
-		Buttons[ (int32)EControllerHand::Left ][ ESteamVRControllerButton::Grip ] = SteamVRControllerKeyNames::BackLeft;
-		Buttons[ (int32)EControllerHand::Left ][ ESteamVRControllerButton::TouchPadUp ] = FGamepadKeyNames::LeftStickUp;
-		Buttons[ (int32)EControllerHand::Left ][ ESteamVRControllerButton::TouchPadDown ] = FGamepadKeyNames::LeftStickDown;
-		Buttons[ (int32)EControllerHand::Left ][ ESteamVRControllerButton::TouchPadLeft ] = FGamepadKeyNames::LeftStickLeft;
-		Buttons[ (int32)EControllerHand::Left ][ ESteamVRControllerButton::TouchPadRight ] = FGamepadKeyNames::LeftStickRight;
+		Buttons[ (int32)EControllerHand::Left ][ ESteamVRControllerButton::TriggerPress ] = FGamepadKeyNames::MotionController_Left_Trigger;
+		Buttons[ (int32)EControllerHand::Left ][ ESteamVRControllerButton::Grip ] = FGamepadKeyNames::MotionController_Left_Grip1;
+		Buttons[ (int32)EControllerHand::Left ][ ESteamVRControllerButton::TouchPadUp ] = FGamepadKeyNames::MotionController_Left_FaceButton1;
+		Buttons[ (int32)EControllerHand::Left ][ ESteamVRControllerButton::TouchPadDown ] = FGamepadKeyNames::MotionController_Left_FaceButton3;
+		Buttons[ (int32)EControllerHand::Left ][ ESteamVRControllerButton::TouchPadLeft ] = FGamepadKeyNames::MotionController_Left_FaceButton4;
+		Buttons[ (int32)EControllerHand::Left ][ ESteamVRControllerButton::TouchPadRight ] = FGamepadKeyNames::MotionController_Left_FaceButton2;
 
 		Buttons[ (int32)EControllerHand::Right ][ ESteamVRControllerButton::System ] = FGamepadKeyNames::SpecialRight;
-		Buttons[ (int32)EControllerHand::Right ][ ESteamVRControllerButton::ApplicationMenu ] = FGamepadKeyNames::RightShoulder;
-		Buttons[ (int32)EControllerHand::Right ][ ESteamVRControllerButton::TouchPadPress ] = FGamepadKeyNames::RightThumb;
+		Buttons[ (int32)EControllerHand::Right ][ ESteamVRControllerButton::ApplicationMenu ] = FGamepadKeyNames::MotionController_Right_Shoulder;
+		Buttons[ (int32)EControllerHand::Right ][ ESteamVRControllerButton::TouchPadPress ] = FGamepadKeyNames::MotionController_Right_Thumbstick;
 		Buttons[ (int32)EControllerHand::Right ][ ESteamVRControllerButton::TouchPadTouch ] = SteamVRControllerKeyNames::Touch1;
-		Buttons[ (int32)EControllerHand::Right ][ ESteamVRControllerButton::TriggerPress ] = FGamepadKeyNames::RightTriggerThreshold;
-		Buttons[ (int32)EControllerHand::Right ][ ESteamVRControllerButton::Grip ] = SteamVRControllerKeyNames::BackRight;
-		Buttons[ (int32)EControllerHand::Right ][ ESteamVRControllerButton::TouchPadUp ] = FGamepadKeyNames::RightStickUp;
-		Buttons[ (int32)EControllerHand::Right ][ ESteamVRControllerButton::TouchPadDown ] = FGamepadKeyNames::RightStickDown;
-		Buttons[ (int32)EControllerHand::Right ][ ESteamVRControllerButton::TouchPadLeft ] = FGamepadKeyNames::RightStickLeft;
-		Buttons[ (int32)EControllerHand::Right ][ ESteamVRControllerButton::TouchPadRight ] = FGamepadKeyNames::RightStickRight;
+		Buttons[ (int32)EControllerHand::Right ][ ESteamVRControllerButton::TriggerPress ] = FGamepadKeyNames::MotionController_Right_Trigger;
+		Buttons[ (int32)EControllerHand::Right ][ ESteamVRControllerButton::Grip ] = FGamepadKeyNames::MotionController_Right_Grip1;
+		Buttons[ (int32)EControllerHand::Right ][ ESteamVRControllerButton::TouchPadUp ] = FGamepadKeyNames::MotionController_Right_FaceButton1;
+		Buttons[ (int32)EControllerHand::Right ][ ESteamVRControllerButton::TouchPadDown ] = FGamepadKeyNames::MotionController_Right_FaceButton3;
+		Buttons[ (int32)EControllerHand::Right ][ ESteamVRControllerButton::TouchPadLeft ] = FGamepadKeyNames::MotionController_Right_FaceButton4;
+		Buttons[ (int32)EControllerHand::Right ][ ESteamVRControllerButton::TouchPadRight ] = FGamepadKeyNames::MotionController_Right_FaceButton2;
 	}
 
 	virtual ~FSteamVRController()
@@ -154,7 +153,7 @@ public:
 						continue;
 					}
 
-					DeviceToControllerMap[DeviceIndex] = NumControllersMapped;
+					DeviceToControllerMap[DeviceIndex] = FMath::FloorToInt(NumControllersMapped / (int32)EControllerHand::TotalHandCount);
 					ControllerToDeviceMap[NumControllersMapped] = DeviceIndex;
 					++NumControllersMapped;
 
@@ -174,7 +173,7 @@ public:
 
 				// get the controller index for this device
 				int32 ControllerIndex = DeviceToControllerMap[DeviceIndex];
-				FControllerState& ControllerState = ControllerStates[ ControllerIndex ];
+				FControllerState& ControllerState = ControllerStates[ DeviceIndex ];
 
 				if (VRSystem->GetControllerState(DeviceIndex, &VRControllerState))
 				{
@@ -191,33 +190,45 @@ public:
 						CurrentStates[ ESteamVRControllerButton::Grip ] = !!(VRControllerState.ulButtonPressed & vr::ButtonMaskFromId(vr::k_EButton_Grip));
 
 						// If the touchpad isn't currently pressed or touched, zero put both of the axes
-						if (!CurrentStates[ ESteamVRControllerButton::TouchPadTouch ] ||	// @todo steamvr: This should actually be a '&&', not an '||'.  For now we require button to be pressed to have non-zeroed data.  This is hack to workaround faulty hardware that reports the touchpad as pressed all of the time
-							!CurrentStates[ ESteamVRControllerButton::TouchPadPress ] )
+						if (!CurrentStates[ ESteamVRControllerButton::TouchPadTouch ])
 						{
-							VRControllerState.rAxis[TOUCHPAD_AXIS].y = 0.0f;
-							VRControllerState.rAxis[TOUCHPAD_AXIS].x = 0.0f;
+ 							VRControllerState.rAxis[TOUCHPAD_AXIS].y = 0.0f;
+ 							VRControllerState.rAxis[TOUCHPAD_AXIS].x = 0.0f;
 						}
 
-						CurrentStates[ ESteamVRControllerButton::TouchPadUp ] = !!(VRControllerState.rAxis[TOUCHPAD_AXIS].y > TOUCHPAD_DEADZONE);
-						CurrentStates[ ESteamVRControllerButton::TouchPadDown ] = !!(VRControllerState.rAxis[TOUCHPAD_AXIS].y < -TOUCHPAD_DEADZONE);
-						CurrentStates[ ESteamVRControllerButton::TouchPadLeft ] = !!(VRControllerState.rAxis[TOUCHPAD_AXIS].x < -TOUCHPAD_DEADZONE);
-						CurrentStates[ ESteamVRControllerButton::TouchPadRight ] = !!(VRControllerState.rAxis[TOUCHPAD_AXIS].x > TOUCHPAD_DEADZONE);
+						// D-pad emulation
+						const FVector2D TouchDir = FVector2D(VRControllerState.rAxis[TOUCHPAD_AXIS].x, VRControllerState.rAxis[TOUCHPAD_AXIS].y).GetSafeNormal();
+						const FVector2D UpDir(0.f, 1.f);
+						const FVector2D RightDir(1.f, 0.f);
+
+						const float VerticalDot = TouchDir | UpDir;
+						const float RightDot = TouchDir | RightDir;
+
+						const bool bPressed = !TouchDir.IsNearlyZero() && CurrentStates[ ESteamVRControllerButton::TouchPadPress ];
+						
+						CurrentStates[ ESteamVRControllerButton::TouchPadUp ]		= bPressed && (VerticalDot >= DOT_45DEG);
+						CurrentStates[ ESteamVRControllerButton::TouchPadDown ]		= bPressed && (VerticalDot <= -DOT_45DEG);
+						CurrentStates[ ESteamVRControllerButton::TouchPadLeft ]		= bPressed && (RightDot <= -DOT_45DEG);
+						CurrentStates[ ESteamVRControllerButton::TouchPadRight ]	= bPressed && (RightDot >= DOT_45DEG);
 
 						if ( ControllerState.TouchPadXAnalog != VRControllerState.rAxis[TOUCHPAD_AXIS].x)
 						{
-							MessageHandler->OnControllerAnalog(FGamepadKeyNames::LeftAnalogX, ControllerIndex, VRControllerState.rAxis[TOUCHPAD_AXIS].x);
+							const FGamepadKeyNames::Type AxisButton = (ControllerState.Hand == EControllerHand::Left) ? FGamepadKeyNames::MotionController_Left_Thumbstick_X : FGamepadKeyNames::MotionController_Right_Thumbstick_X;
+							MessageHandler->OnControllerAnalog(AxisButton, ControllerIndex, VRControllerState.rAxis[TOUCHPAD_AXIS].x);
 							ControllerState.TouchPadXAnalog = VRControllerState.rAxis[TOUCHPAD_AXIS].x;
 						}
 
 						if ( ControllerState.TouchPadYAnalog != VRControllerState.rAxis[TOUCHPAD_AXIS].y)
 						{
-							MessageHandler->OnControllerAnalog(FGamepadKeyNames::LeftAnalogY, ControllerIndex, VRControllerState.rAxis[TOUCHPAD_AXIS].y);
+							const FGamepadKeyNames::Type AxisButton = (ControllerState.Hand == EControllerHand::Left) ? FGamepadKeyNames::MotionController_Left_Thumbstick_Y : FGamepadKeyNames::MotionController_Right_Thumbstick_Y;
+							MessageHandler->OnControllerAnalog(AxisButton, ControllerIndex, VRControllerState.rAxis[TOUCHPAD_AXIS].y);
 							ControllerState.TouchPadYAnalog = VRControllerState.rAxis[TOUCHPAD_AXIS].y;
 						}
 
 						if ( ControllerState.TriggerAnalog != VRControllerState.rAxis[TRIGGER_AXIS].x)
 						{
-							MessageHandler->OnControllerAnalog(FGamepadKeyNames::LeftTriggerAnalog, ControllerIndex, VRControllerState.rAxis[TRIGGER_AXIS].x);
+							const FGamepadKeyNames::Type AxisButton = (ControllerState.Hand == EControllerHand::Left) ? FGamepadKeyNames::MotionController_Left_TriggerAxis : FGamepadKeyNames::MotionController_Right_TriggerAxis;
+							MessageHandler->OnControllerAnalog(AxisButton, ControllerIndex, VRControllerState.rAxis[TRIGGER_AXIS].x);
 							ControllerState.TriggerAnalog = VRControllerState.rAxis[TRIGGER_AXIS].x;
 						}
 
diff --git a/Engine/Source/Runtime/Core/Private/GenericPlatform/GenericApplication.cpp b/Engine/Source/Runtime/Core/Private/GenericPlatform/GenericApplication.cpp
--- a/Engine/Source/Runtime/Core/Private/GenericPlatform/GenericApplication.cpp
+++ b/Engine/Source/Runtime/Core/Private/GenericPlatform/GenericApplication.cpp
@@ -41,6 +41,60 @@ const FGamepadKeyNames::Type FGamepadKeyNames::RightStickDown("Gamepad_RightStic
 const FGamepadKeyNames::Type FGamepadKeyNames::RightStickRight("Gamepad_RightStick_Right");
 const FGamepadKeyNames::Type FGamepadKeyNames::RightStickLeft("Gamepad_RightStick_Left");
 
+const FGamepadKeyNames::Type FGamepadKeyNames::MotionController_Left_FaceButton1("MotionController_Left_FaceButton1");
+const FGamepadKeyNames::Type FGamepadKeyNames::MotionController_Left_FaceButton2("MotionController_Left_FaceButton2");
+const FGamepadKeyNames::Type FGamepadKeyNames::MotionController_Left_FaceButton3("MotionController_Left_FaceButton3");
+const FGamepadKeyNames::Type FGamepadKeyNames::MotionController_Left_FaceButton4("MotionController_Left_FaceButton4");
+const FGamepadKeyNames::Type FGamepadKeyNames::MotionController_Left_FaceButton5("MotionController_Left_FaceButton5");
+const FGamepadKeyNames::Type FGamepadKeyNames::MotionController_Left_FaceButton6("MotionController_Left_FaceButton6");
+const FGamepadKeyNames::Type FGamepadKeyNames::MotionController_Left_FaceButton7("MotionController_Left_FaceButton7");
+const FGamepadKeyNames::Type FGamepadKeyNames::MotionController_Left_FaceButton8("MotionController_Left_FaceButton8");
+
+const FGamepadKeyNames::Type FGamepadKeyNames::MotionController_Left_Shoulder("MotionController_Left_Shoulder");
+const FGamepadKeyNames::Type FGamepadKeyNames::MotionController_Left_Trigger("MotionController_Left_Trigger");
+
+const FGamepadKeyNames::Type FGamepadKeyNames::MotionController_Left_Grip1("MotionController_Left_Grip1");
+const FGamepadKeyNames::Type FGamepadKeyNames::MotionController_Left_Grip2("MotionController_Left_Grip2");
+
+const FGamepadKeyNames::Type FGamepadKeyNames::MotionController_Left_Thumbstick("MotionController_Left_Thumbstick");
+const FGamepadKeyNames::Type FGamepadKeyNames::MotionController_Left_Thumbstick_Up("MotionController_Left_Thumbstick_Up");
+const FGamepadKeyNames::Type FGamepadKeyNames::MotionController_Left_Thumbstick_Down("MotionController_Left_Thumbstick_Down");
+const FGamepadKeyNames::Type FGamepadKeyNames::MotionController_Left_Thumbstick_Left("MotionController_Left_Thumbstick_Left");
+const FGamepadKeyNames::Type FGamepadKeyNames::MotionController_Left_Thumbstick_Right("MotionController_Left_Thumbstick_Right");
+
+	//		Right Controller
+const FGamepadKeyNames::Type FGamepadKeyNames::MotionController_Right_FaceButton1("MotionController_Right_FaceButton1");
+const FGamepadKeyNames::Type FGamepadKeyNames::MotionController_Right_FaceButton2("MotionController_Right_FaceButton2");
+const FGamepadKeyNames::Type FGamepadKeyNames::MotionController_Right_FaceButton3("MotionController_Right_FaceButton3");
+const FGamepadKeyNames::Type FGamepadKeyNames::MotionController_Right_FaceButton4("MotionController_Right_FaceButton4");
+const FGamepadKeyNames::Type FGamepadKeyNames::MotionController_Right_FaceButton5("MotionController_Right_FaceButton5");
+const FGamepadKeyNames::Type FGamepadKeyNames::MotionController_Right_FaceButton6("MotionController_Right_FaceButton6");
+const FGamepadKeyNames::Type FGamepadKeyNames::MotionController_Right_FaceButton7("MotionController_Right_FaceButton7");
+const FGamepadKeyNames::Type FGamepadKeyNames::MotionController_Right_FaceButton8("MotionController_Right_FaceButton8");
+
+const FGamepadKeyNames::Type FGamepadKeyNames::MotionController_Right_Shoulder("MotionController_Right_Shoulder");
+const FGamepadKeyNames::Type FGamepadKeyNames::MotionController_Right_Trigger("MotionController_Right_Trigger");
+
+const FGamepadKeyNames::Type FGamepadKeyNames::MotionController_Right_Grip1("MotionController_Right_Grip1");
+const FGamepadKeyNames::Type FGamepadKeyNames::MotionController_Right_Grip2("MotionController_Right_Grip2");
+
+const FGamepadKeyNames::Type FGamepadKeyNames::MotionController_Right_Thumbstick("MotionController_Right_Thumbstick");
+const FGamepadKeyNames::Type FGamepadKeyNames::MotionController_Right_Thumbstick_Up("MotionController_Right_Thumbstick_Up");
+const FGamepadKeyNames::Type FGamepadKeyNames::MotionController_Right_Thumbstick_Down("MotionController_Right_Thumbstick_Down");
+const FGamepadKeyNames::Type FGamepadKeyNames::MotionController_Right_Thumbstick_Left("MotionController_Right_Thumbstick_Left");
+const FGamepadKeyNames::Type FGamepadKeyNames::MotionController_Right_Thumbstick_Right("MotionController_Right_Thumbstick_Right");
+
+	//   Motion Controller Axes
+	//		Left Controller
+const FGamepadKeyNames::Type FGamepadKeyNames::MotionController_Left_Thumbstick_X("MotionController_Left_Thumbstick_X");
+const FGamepadKeyNames::Type FGamepadKeyNames::MotionController_Left_Thumbstick_Y("MotionController_Left_Thumbstick_Y");
+const FGamepadKeyNames::Type FGamepadKeyNames::MotionController_Left_TriggerAxis("MotionController_Left_TriggerAxis");
+
+	//		Right Controller
+const FGamepadKeyNames::Type FGamepadKeyNames::MotionController_Right_Thumbstick_X("MotionController_Right_Thumbstick_X");
+const FGamepadKeyNames::Type FGamepadKeyNames::MotionController_Right_Thumbstick_Y("MotionController_Right_Thumbstick_Y");
+const FGamepadKeyNames::Type FGamepadKeyNames::MotionController_Right_TriggerAxis("MotionController_Right_TriggerAxis");
+
 void FDisplayMetrics::PrintToLog() const
 {
 	UE_LOG(LogInit, Log, TEXT("Display metrics:"));
diff --git a/Engine/Source/Runtime/Core/Public/GenericPlatform/GenericApplicationMessageHandler.h b/Engine/Source/Runtime/Core/Public/GenericPlatform/GenericApplicationMessageHandler.h
--- a/Engine/Source/Runtime/Core/Public/GenericPlatform/GenericApplicationMessageHandler.h
+++ b/Engine/Source/Runtime/Core/Public/GenericPlatform/GenericApplicationMessageHandler.h
@@ -63,6 +63,62 @@ struct CORE_API FGamepadKeyNames
 	static const FName RightStickDown;
 	static const FName RightStickRight;
 	static const FName RightStickLeft;
+
+	// Motion Controls
+	//		Left Controller
+	static const FName MotionController_Left_FaceButton1;
+	static const FName MotionController_Left_FaceButton2;
+	static const FName MotionController_Left_FaceButton3;
+	static const FName MotionController_Left_FaceButton4;
+	static const FName MotionController_Left_FaceButton5;
+	static const FName MotionController_Left_FaceButton6;
+	static const FName MotionController_Left_FaceButton7;
+	static const FName MotionController_Left_FaceButton8;
+
+	static const FName MotionController_Left_Shoulder;
+	static const FName MotionController_Left_Trigger;
+
+	static const FName MotionController_Left_Grip1;
+	static const FName MotionController_Left_Grip2;
+
+	static const FName MotionController_Left_Thumbstick;
+	static const FName MotionController_Left_Thumbstick_Up;
+	static const FName MotionController_Left_Thumbstick_Down;
+	static const FName MotionController_Left_Thumbstick_Left;
+	static const FName MotionController_Left_Thumbstick_Right;
+
+	//		Right Controller
+	static const FName MotionController_Right_FaceButton1;
+	static const FName MotionController_Right_FaceButton2;
+	static const FName MotionController_Right_FaceButton3;
+	static const FName MotionController_Right_FaceButton4;
+	static const FName MotionController_Right_FaceButton5;
+	static const FName MotionController_Right_FaceButton6;
+	static const FName MotionController_Right_FaceButton7;
+	static const FName MotionController_Right_FaceButton8;
+
+	static const FName MotionController_Right_Shoulder;
+	static const FName MotionController_Right_Trigger;
+
+	static const FName MotionController_Right_Grip1;
+	static const FName MotionController_Right_Grip2;
+
+	static const FName MotionController_Right_Thumbstick;
+	static const FName MotionController_Right_Thumbstick_Up;
+	static const FName MotionController_Right_Thumbstick_Down;
+	static const FName MotionController_Right_Thumbstick_Left;
+	static const FName MotionController_Right_Thumbstick_Right;
+
+	//   Motion Controller Axes
+	//		Left Controller
+	static const FName MotionController_Left_Thumbstick_X;
+	static const FName MotionController_Left_Thumbstick_Y;
+	static const FName MotionController_Left_TriggerAxis;
+
+	//		Right Controller
+	static const FName MotionController_Right_Thumbstick_X;
+	static const FName MotionController_Right_Thumbstick_Y;
+	static const FName MotionController_Right_TriggerAxis;
 };
 
 namespace EWindowActivation
diff --git a/Engine/Source/Runtime/InputCore/Classes/InputCoreTypes.h b/Engine/Source/Runtime/InputCore/Classes/InputCoreTypes.h
--- a/Engine/Source/Runtime/InputCore/Classes/InputCoreTypes.h
+++ b/Engine/Source/Runtime/InputCore/Classes/InputCoreTypes.h
@@ -353,6 +353,62 @@ struct INPUTCORE_API EKeys
 	static const FKey Gesture_Pinch;
 	static const FKey Gesture_Flick;
 
+	// Motion Controllers
+	//		Left Controller
+	static const FKey MotionController_Left_FaceButton1;
+	static const FKey MotionController_Left_FaceButton2;
+	static const FKey MotionController_Left_FaceButton3;
+	static const FKey MotionController_Left_FaceButton4;
+	static const FKey MotionController_Left_FaceButton5;
+	static const FKey MotionController_Left_FaceButton6;
+	static const FKey MotionController_Left_FaceButton7;
+	static const FKey MotionController_Left_FaceButton8;
+
+	static const FKey MotionController_Left_Shoulder;
+	static const FKey MotionController_Left_Trigger;
+
+	static const FKey MotionController_Left_Grip1;
+	static const FKey MotionController_Left_Grip2;
+
+	static const FKey MotionController_Left_Thumbstick;
+	static const FKey MotionController_Left_Thumbstick_Up;
+	static const FKey MotionController_Left_Thumbstick_Down;
+	static const FKey MotionController_Left_Thumbstick_Left;
+	static const FKey MotionController_Left_Thumbstick_Right;
+
+	//		Right Controller
+	static const FKey MotionController_Right_FaceButton1;
+	static const FKey MotionController_Right_FaceButton2;
+	static const FKey MotionController_Right_FaceButton3;
+	static const FKey MotionController_Right_FaceButton4;
+	static const FKey MotionController_Right_FaceButton5;
+	static const FKey MotionController_Right_FaceButton6;
+	static const FKey MotionController_Right_FaceButton7;
+	static const FKey MotionController_Right_FaceButton8;
+
+	static const FKey MotionController_Right_Shoulder;
+	static const FKey MotionController_Right_Trigger;
+
+	static const FKey MotionController_Right_Grip1;
+	static const FKey MotionController_Right_Grip2;
+
+	static const FKey MotionController_Right_Thumbstick;
+	static const FKey MotionController_Right_Thumbstick_Up;
+	static const FKey MotionController_Right_Thumbstick_Down;
+	static const FKey MotionController_Right_Thumbstick_Left;
+	static const FKey MotionController_Right_Thumbstick_Right;
+	 
+	//   Motion Controller Axes
+	//		Left Controller
+	static const FKey MotionController_Left_Thumbstick_X;
+	static const FKey MotionController_Left_Thumbstick_Y;
+	static const FKey MotionController_Left_TriggerAxis;
+
+	//		Right Controller
+	static const FKey MotionController_Right_Thumbstick_X;
+	static const FKey MotionController_Right_Thumbstick_Y;
+	static const FKey MotionController_Right_TriggerAxis;
+
 	// PS4-specific
 	static const FKey PS4_Special;
 
diff --git a/Engine/Source/Runtime/InputCore/Private/InputCoreTypes.cpp b/Engine/Source/Runtime/InputCore/Private/InputCoreTypes.cpp
--- a/Engine/Source/Runtime/InputCore/Private/InputCoreTypes.cpp
+++ b/Engine/Source/Runtime/InputCore/Private/InputCoreTypes.cpp
@@ -218,6 +218,62 @@ const FKey EKeys::TouchKeys[NUM_TOUCH_KEYS] =
 const FKey EKeys::Gesture_Pinch("Gesture_Pinch");
 const FKey EKeys::Gesture_Flick("Gesture_Flick");
 
+// Motion Controllers
+//		Left Controller
+const FKey EKeys::MotionController_Left_FaceButton1("MotionController_Left_FaceButton1");
+const FKey EKeys::MotionController_Left_FaceButton2("MotionController_Left_FaceButton2");
+const FKey EKeys::MotionController_Left_FaceButton3("MotionController_Left_FaceButton3");
+const FKey EKeys::MotionController_Left_FaceButton4("MotionController_Left_FaceButton4");
+const FKey EKeys::MotionController_Left_FaceButton5("MotionController_Left_FaceButton5");
+const FKey EKeys::MotionController_Left_FaceButton6("MotionController_Left_FaceButton6");
+const FKey EKeys::MotionController_Left_FaceButton7("MotionController_Left_FaceButton7");
+const FKey EKeys::MotionController_Left_FaceButton8("MotionController_Left_FaceButton8");
+
+const FKey EKeys::MotionController_Left_Shoulder("MotionController_Left_Shoulder");
+const FKey EKeys::MotionController_Left_Trigger("MotionController_Left_Trigger");
+
+const FKey EKeys::MotionController_Left_Grip1("MotionController_Left_Grip1");
+const FKey EKeys::MotionController_Left_Grip2("MotionController_Left_Grip2");
+
+const FKey EKeys::MotionController_Left_Thumbstick("MotionController_Left_Thumbstick");
+const FKey EKeys::MotionController_Left_Thumbstick_Up("MotionController_Left_Thumbstick_Up");
+const FKey EKeys::MotionController_Left_Thumbstick_Down("MotionController_Left_Thumbstick_Down");
+const FKey EKeys::MotionController_Left_Thumbstick_Left("MotionController_Left_Thumbstick_Left");
+const FKey EKeys::MotionController_Left_Thumbstick_Right("MotionController_Left_Thumbstick_Right");
+
+//		Right Controller
+const FKey EKeys::MotionController_Right_FaceButton1("MotionController_Right_FaceButton1");
+const FKey EKeys::MotionController_Right_FaceButton2("MotionController_Right_FaceButton2");
+const FKey EKeys::MotionController_Right_FaceButton3("MotionController_Right_FaceButton3");
+const FKey EKeys::MotionController_Right_FaceButton4("MotionController_Right_FaceButton4");
+const FKey EKeys::MotionController_Right_FaceButton5("MotionController_Right_FaceButton5");
+const FKey EKeys::MotionController_Right_FaceButton6("MotionController_Right_FaceButton6");
+const FKey EKeys::MotionController_Right_FaceButton7("MotionController_Right_FaceButton7");
+const FKey EKeys::MotionController_Right_FaceButton8("MotionController_Right_FaceButton8");
+
+const FKey EKeys::MotionController_Right_Shoulder("MotionController_Right_Shoulder");
+const FKey EKeys::MotionController_Right_Trigger("MotionController_Right_Trigger");
+
+const FKey EKeys::MotionController_Right_Grip1("MotionController_Right_Grip1");
+const FKey EKeys::MotionController_Right_Grip2("MotionController_Right_Grip2");
+
+const FKey EKeys::MotionController_Right_Thumbstick("MotionController_Right_Thumbstick");
+const FKey EKeys::MotionController_Right_Thumbstick_Up("MotionController_Right_Thumbstick_Up");
+const FKey EKeys::MotionController_Right_Thumbstick_Down("MotionController_Right_Thumbstick_Down");
+const FKey EKeys::MotionController_Right_Thumbstick_Left("MotionController_Right_Thumbstick_Left");
+const FKey EKeys::MotionController_Right_Thumbstick_Right("MotionController_Right_Thumbstick_Right");
+
+//   Motion Controller Axes
+//		Left Controller
+const FKey EKeys::MotionController_Left_Thumbstick_X("MotionController_Left_Thumbstick_X");
+const FKey EKeys::MotionController_Left_Thumbstick_Y("MotionController_Left_Thumbstick_Y");
+const FKey EKeys::MotionController_Left_TriggerAxis("MotionController_Left_TriggerAxis");
+
+//		Right Controller
+const FKey EKeys::MotionController_Right_Thumbstick_X("MotionController_Right_Thumbstick_X");
+const FKey EKeys::MotionController_Right_Thumbstick_Y("MotionController_Right_Thumbstick_Y");
+const FKey EKeys::MotionController_Right_TriggerAxis("MotionController_Right_TriggerAxis");
+
 // PS4-specific
 const FKey EKeys::PS4_Special("PS4_Special");
 
@@ -511,7 +567,66 @@ void EKeys::Initialize()
 	AddKey(FKeyDetails(EKeys::Gesture_Pinch, LOCTEXT("Gesture_Pinch", "Pinch"), 0, "Gesture"));
 	AddKey(FKeyDetails(EKeys::Gesture_Flick, LOCTEXT("Gesture_Flick", "Flick"), 0, "Gesture"));
 
+	// Motion Controllers
+	AddMenuCategoryDisplayInfo("Steam", LOCTEXT("MotionControllerSubCategory", "Motion Controller"), TEXT("GraphEditor.PadEvent_16x"));
+
+	//	Left Controller
+	AddKey(FKeyDetails(EKeys::MotionController_Left_FaceButton1, LOCTEXT("MotionController_Left_FaceButton1", "MotionController (L) FaceButton1"), FKeyDetails::GamepadKey));
+	AddKey(FKeyDetails(EKeys::MotionController_Left_FaceButton2, LOCTEXT("MotionController_Left_FaceButton2", "MotionController (L) FaceButton2"), FKeyDetails::GamepadKey));
+	AddKey(FKeyDetails(EKeys::MotionController_Left_FaceButton3, LOCTEXT("MotionController_Left_FaceButton3", "MotionController (L) FaceButton3"), FKeyDetails::GamepadKey));
+	AddKey(FKeyDetails(EKeys::MotionController_Left_FaceButton4, LOCTEXT("MotionController_Left_FaceButton4", "MotionController (L) FaceButton4"), FKeyDetails::GamepadKey));
+	AddKey(FKeyDetails(EKeys::MotionController_Left_FaceButton5, LOCTEXT("MotionController_Left_FaceButton5", "MotionController (L) FaceButton5"), FKeyDetails::GamepadKey));
+	AddKey(FKeyDetails(EKeys::MotionController_Left_FaceButton6, LOCTEXT("MotionController_Left_FaceButton6", "MotionController (L) FaceButton6"), FKeyDetails::GamepadKey));
+	AddKey(FKeyDetails(EKeys::MotionController_Left_FaceButton7, LOCTEXT("MotionController_Left_FaceButton7", "MotionController (L) FaceButton7"), FKeyDetails::GamepadKey));
+	AddKey(FKeyDetails(EKeys::MotionController_Left_FaceButton8, LOCTEXT("MotionController_Left_FaceButton8", "MotionController (L) FaceButton8"), FKeyDetails::GamepadKey));
+	
+	AddKey(FKeyDetails(EKeys::MotionController_Left_Shoulder, LOCTEXT("MotionController_Left_Shoulder", "MotionController (L) Shoulder"), FKeyDetails::GamepadKey));
+	AddKey(FKeyDetails(EKeys::MotionController_Left_Trigger, LOCTEXT("MotionController_Left_Trigger", "MotionController (L) Trigger"), FKeyDetails::GamepadKey));
+	
+	AddKey(FKeyDetails(EKeys::MotionController_Left_Grip1, LOCTEXT("MotionController_Left_Grip1", "MotionController (L) Grip1"), FKeyDetails::GamepadKey));
+	AddKey(FKeyDetails(EKeys::MotionController_Left_Grip2, LOCTEXT("MotionController_Left_Grip2", "MotionController (L) Grip2"), FKeyDetails::GamepadKey));
+	
+	AddKey(FKeyDetails(EKeys::MotionController_Left_Thumbstick, LOCTEXT("MotionController_Left_Thumbstick", "MotionController (L) Thumbstick"), FKeyDetails::GamepadKey));
+	AddKey(FKeyDetails(EKeys::MotionController_Left_Thumbstick_Up, LOCTEXT("MotionController_Left_Thumbstick_Up", "MotionController (L) Thumbstick Up"), FKeyDetails::GamepadKey));
+	AddKey(FKeyDetails(EKeys::MotionController_Left_Thumbstick_Down, LOCTEXT("MotionController_Left_Thumbstick_Down", "MotionController (L) Thumbstick Down"), FKeyDetails::GamepadKey));
+	AddKey(FKeyDetails(EKeys::MotionController_Left_Thumbstick_Left, LOCTEXT("MotionController_Left_Thumbstick_Left", "MotionController (L) Thumbstick Left"), FKeyDetails::GamepadKey));
+	AddKey(FKeyDetails(EKeys::MotionController_Left_Thumbstick_Right, LOCTEXT("MotionController_Left_Thumbstick_Right", "MotionController (L) Thumbstick Right"), FKeyDetails::GamepadKey));
+	
+	//	Right Controller
+	AddKey(FKeyDetails(EKeys::MotionController_Right_FaceButton1, LOCTEXT("MotionController_Right_FaceButton1", "MotionController (R) FaceButton1"), FKeyDetails::GamepadKey));
+	AddKey(FKeyDetails(EKeys::MotionController_Right_FaceButton2, LOCTEXT("MotionController_Right_FaceButton2", "MotionController (R) FaceButton2"), FKeyDetails::GamepadKey));
+	AddKey(FKeyDetails(EKeys::MotionController_Right_FaceButton3, LOCTEXT("MotionController_Right_FaceButton3", "MotionController (R) FaceButton3"), FKeyDetails::GamepadKey));
+	AddKey(FKeyDetails(EKeys::MotionController_Right_FaceButton4, LOCTEXT("MotionController_Right_FaceButton4", "MotionController (R) FaceButton4"), FKeyDetails::GamepadKey));
+	AddKey(FKeyDetails(EKeys::MotionController_Right_FaceButton5, LOCTEXT("MotionController_Right_FaceButton5", "MotionController (R) FaceButton5"), FKeyDetails::GamepadKey));
+	AddKey(FKeyDetails(EKeys::MotionController_Right_FaceButton6, LOCTEXT("MotionController_Right_FaceButton6", "MotionController (R) FaceButton6"), FKeyDetails::GamepadKey));
+	AddKey(FKeyDetails(EKeys::MotionController_Right_FaceButton7, LOCTEXT("MotionController_Right_FaceButton7", "MotionController (R) FaceButton7"), FKeyDetails::GamepadKey));
+	AddKey(FKeyDetails(EKeys::MotionController_Right_FaceButton8, LOCTEXT("MotionController_Right_FaceButton8", "MotionController (R) FaceButton8"), FKeyDetails::GamepadKey));
+	
+	AddKey(FKeyDetails(EKeys::MotionController_Right_Shoulder, LOCTEXT("MotionController_Right_Shoulder", "MotionController (R) Shoulder"), FKeyDetails::GamepadKey));
+	AddKey(FKeyDetails(EKeys::MotionController_Right_Trigger, LOCTEXT("MotionController_Right_Trigger", "MotionController (R) Trigger"), FKeyDetails::GamepadKey));
+	
+	AddKey(FKeyDetails(EKeys::MotionController_Right_Grip1, LOCTEXT("MotionController_Right_Grip1", "MotionController (R) Grip1"), FKeyDetails::GamepadKey));
+	AddKey(FKeyDetails(EKeys::MotionController_Right_Grip2, LOCTEXT("MotionController_Right_Grip2", "MotionController (R) Grip2"), FKeyDetails::GamepadKey));
+	
+	AddKey(FKeyDetails(EKeys::MotionController_Right_Thumbstick, LOCTEXT("MotionController_Right_Thumbstick", "MotionController (R) Thumbstick"), FKeyDetails::GamepadKey));
+	AddKey(FKeyDetails(EKeys::MotionController_Right_Thumbstick_Up, LOCTEXT("MotionController_Right_Thumbstick_Up", "MotionController (R) Thumbstick Up"), FKeyDetails::GamepadKey));
+	AddKey(FKeyDetails(EKeys::MotionController_Right_Thumbstick_Down, LOCTEXT("MotionController_Right_Thumbstick_Down", "MotionController (R) Thumbstick Down"), FKeyDetails::GamepadKey));
+	AddKey(FKeyDetails(EKeys::MotionController_Right_Thumbstick_Left, LOCTEXT("MotionController_Right_Thumbstick_Left", "MotionController (R) Thumbstick Left"), FKeyDetails::GamepadKey));
+	AddKey(FKeyDetails(EKeys::MotionController_Right_Thumbstick_Right, LOCTEXT("MotionController_Right_Thumbstick_Right", "MotionController (R) Thumbstick Right"), FKeyDetails::GamepadKey));
+	
+	//   Motion Controller Axes
+	//	Left Controller
+	AddKey(FKeyDetails(EKeys::MotionController_Left_Thumbstick_X, LOCTEXT("MotionController_Left_Thumbstick_X", "MotionController (L) Thumbstick X"), FKeyDetails::GamepadKey | FKeyDetails::FloatAxis));
+	AddKey(FKeyDetails(EKeys::MotionController_Left_Thumbstick_Y, LOCTEXT("MotionController_Left_Thumbstick_Y", "MotionController (L) Thumbstick Y"), FKeyDetails::GamepadKey | FKeyDetails::FloatAxis));
+	AddKey(FKeyDetails(EKeys::MotionController_Left_TriggerAxis, LOCTEXT("MotionController_Left_TriggerAxis", "MotionController (L) TriggerAxis"), FKeyDetails::GamepadKey | FKeyDetails::FloatAxis));
+	
+	//	Right Controller
+	AddKey(FKeyDetails(EKeys::MotionController_Right_Thumbstick_X, LOCTEXT("MotionController_Right_Thumbstick_X", "MotionController (R) Thumbstick X"), FKeyDetails::GamepadKey | FKeyDetails::FloatAxis));
+	AddKey(FKeyDetails(EKeys::MotionController_Right_Thumbstick_Y, LOCTEXT("MotionController_Right_Thumbstick_Y", "MotionController (R) Thumbstick Y"), FKeyDetails::GamepadKey | FKeyDetails::FloatAxis));
+	AddKey(FKeyDetails(EKeys::MotionController_Right_TriggerAxis, LOCTEXT("MotionController_Right_TriggerAxis", "MotionController (R) TriggerAxis"), FKeyDetails::GamepadKey | FKeyDetails::FloatAxis));
+
 	// PS4-specific
+	AddMenuCategoryDisplayInfo("Steam", LOCTEXT("PS4SubCategory", "PS4"), TEXT("GraphEditor.PadEvent_16x"));
 	AddKey(FKeyDetails(EKeys::PS4_Special, LOCTEXT("PS4_Special", "PS4_Special"), FKeyDetails::NotBlueprintBindableKey));
 
 	// Steam Controller specific
