Expose MouseWheelAxis to give proper axis access to mouse wheel
Generalize special case handling for axes that don't push values every frame

diff --git a/Engine/Source/Runtime/Engine/Private/Slate/SceneViewport.cpp b/Engine/Source/Runtime/Engine/Private/Slate/SceneViewport.cpp
--- a/Engine/Source/Runtime/Engine/Private/Slate/SceneViewport.cpp
+++ b/Engine/Source/Runtime/Engine/Private/Slate/SceneViewport.cpp
@@ -544,6 +544,7 @@ FReply FSceneViewport::OnMouseWheel( const FGeometry& InGeometry, const FPointer
 		// Pressed and released should be sent
 		ViewportClient->InputKey(this, InMouseEvent.GetUserIndex(), ViewportClientKey, IE_Pressed);
 		ViewportClient->InputKey(this, InMouseEvent.GetUserIndex(), ViewportClientKey, IE_Released);
+		ViewportClient->InputAxis(this, InMouseEvent.GetUserIndex(), EKeys::MouseWheelAxis, InMouseEvent.GetWheelDelta(), FApp::GetDeltaTime());
 	}
 	return CurrentReplyState;
 }
diff --git a/Engine/Source/Runtime/Engine/Private/UserInterface/PlayerInput.cpp b/Engine/Source/Runtime/Engine/Private/UserInterface/PlayerInput.cpp
--- a/Engine/Source/Runtime/Engine/Private/UserInterface/PlayerInput.cpp
+++ b/Engine/Source/Runtime/Engine/Private/UserInterface/PlayerInput.cpp
@@ -802,7 +802,7 @@ void UPlayerInput::ProcessInputStack(const TArray<UInputComponent*>& InputCompon
 			KeyState->EventCounts[EventIndex] = MoveTemp(KeyState->EventAccumulator[EventIndex]);
 		}
 
-		if ( (KeyState->SampleCountAccumulator > 0) || (Key == EKeys::MouseX) || (Key == EKeys::MouseY) )
+		if ( (KeyState->SampleCountAccumulator > 0) || Key.ShouldUpdateAxisWithoutSamples() )
 		{
 			// if we had no samples, we'll assume the state hasn't changed
 			// except for some axes, where no samples means the mouse stopped moving
diff --git a/Engine/Source/Runtime/InputCore/Classes/InputCoreTypes.h b/Engine/Source/Runtime/InputCore/Classes/InputCoreTypes.h
--- a/Engine/Source/Runtime/InputCore/Classes/InputCoreTypes.h
+++ b/Engine/Source/Runtime/InputCore/Classes/InputCoreTypes.h
@@ -37,6 +37,7 @@ struct INPUTCORE_API FKey
 	bool IsFloatAxis() const;
 	bool IsVectorAxis() const;
 	bool IsBindableInBlueprints() const;
+	bool ShouldUpdateAxisWithoutSamples() const;
 	FText GetDisplayName() const;
 	FString ToString() const;
 	FName GetFName() const;
@@ -90,6 +91,7 @@ struct INPUTCORE_API FKeyDetails
 		NotBlueprintBindableKey	= 0x08,
 		FloatAxis				= 0x10,
 		VectorAxis				= 0x20,
+		UpdateAxisWithoutSamples = 0x40,
 
 		NoFlags                 = 0,
 	};
@@ -102,6 +104,7 @@ struct INPUTCORE_API FKeyDetails
 	bool IsFloatAxis() const { return AxisType == EInputAxisType::Float; }
 	bool IsVectorAxis() const { return AxisType == EInputAxisType::Vector; }
 	bool IsBindableInBlueprints() const { return bIsBindableInBlueprints != 0; }
+	bool ShouldUpdateAxisWithoutSamples() const { return bShouldUpdateAxisWithoutSamples != 0; }
 	FName GetMenuCategory() const { return MenuCategory; }
 	FText GetDisplayName() const;
 	const FKey& GetKey() const { return Key; }
@@ -125,6 +128,7 @@ private:
 	int32 bIsGamepadKey:1;
 	int32 bIsMouseButton:1;
 	int32 bIsBindableInBlueprints:1;
+	int32 bShouldUpdateAxisWithoutSamples:1;
 	EInputAxisType AxisType;
 
 };
@@ -164,6 +168,7 @@ struct INPUTCORE_API EKeys
 	static const FKey MouseY;
 	static const FKey MouseScrollUp;
 	static const FKey MouseScrollDown;
+	static const FKey MouseWheelAxis;
 
 	static const FKey LeftMouseButton;
 	static const FKey RightMouseButton;
diff --git a/Engine/Source/Runtime/InputCore/Private/InputCoreTypes.cpp b/Engine/Source/Runtime/InputCore/Private/InputCoreTypes.cpp
--- a/Engine/Source/Runtime/InputCore/Private/InputCoreTypes.cpp
+++ b/Engine/Source/Runtime/InputCore/Private/InputCoreTypes.cpp
@@ -12,6 +12,7 @@ const FKey EKeys::MouseX("MouseX");
 const FKey EKeys::MouseY("MouseY");
 const FKey EKeys::MouseScrollUp("MouseScrollUp");
 const FKey EKeys::MouseScrollDown("MouseScrollDown");
+const FKey EKeys::MouseWheelAxis("MouseWheelAxis");
 
 const FKey EKeys::LeftMouseButton("LeftMouseButton");
 const FKey EKeys::RightMouseButton("RightMouseButton");
@@ -316,6 +317,7 @@ FKeyDetails::FKeyDetails(const FKey InKey, const TAttribute<FText>& InDisplayNam
 	, bIsGamepadKey((InKeyFlags & EKeyFlags::GamepadKey) != 0)
 	, bIsMouseButton((InKeyFlags & EKeyFlags::MouseButton) != 0)
 	, bIsBindableInBlueprints((~InKeyFlags & EKeyFlags::NotBlueprintBindableKey) != 0)
+	, bShouldUpdateAxisWithoutSamples((InKeyFlags & EKeyFlags::UpdateAxisWithoutSamples) != 0)
 	, AxisType(EInputAxisType::None)
 {
 	if ((InKeyFlags & EKeyFlags::FloatAxis) != 0)
@@ -355,11 +357,12 @@ void EKeys::Initialize()
 	AddMenuCategoryDisplayInfo(NAME_MouseCategory, LOCTEXT("MouseSubCategory", "Mouse"), TEXT("GraphEditor.MouseEvent_16x"));
 	AddMenuCategoryDisplayInfo(NAME_KeyboardCategory, LOCTEXT("KeyboardSubCategory", "Keyboard"), TEXT("GraphEditor.KeyEvent_16x"));
 
-	AddKey(FKeyDetails(EKeys::MouseX, LOCTEXT("MouseX", "Mouse X"), FKeyDetails::FloatAxis | FKeyDetails::MouseButton));
-	AddKey(FKeyDetails(EKeys::MouseY, LOCTEXT("MouseY", "Mouse Y"), FKeyDetails::FloatAxis | FKeyDetails::MouseButton));
+	AddKey(FKeyDetails(EKeys::MouseX, LOCTEXT("MouseX", "Mouse X"), FKeyDetails::FloatAxis | FKeyDetails::MouseButton | FKeyDetails::UpdateAxisWithoutSamples));
+	AddKey(FKeyDetails(EKeys::MouseY, LOCTEXT("MouseY", "Mouse Y"), FKeyDetails::FloatAxis | FKeyDetails::MouseButton | FKeyDetails::UpdateAxisWithoutSamples));
+	AddKey(FKeyDetails(EKeys::MouseWheelAxis, LOCTEXT("MouseWheelAxis", "Mouse Wheel Axis"), FKeyDetails::FloatAxis | FKeyDetails::MouseButton | FKeyDetails::UpdateAxisWithoutSamples));
 	AddKey(FKeyDetails(EKeys::MouseScrollUp, LOCTEXT("MouseScrollUp", "Mouse Wheel Up"), FKeyDetails::MouseButton));
 	AddKey(FKeyDetails(EKeys::MouseScrollDown, LOCTEXT("MouseScrollDown", "Mouse Wheel Down"), FKeyDetails::MouseButton));
-
+	
 	AddKey(FKeyDetails(EKeys::LeftMouseButton, LOCTEXT("LeftMouseButton", "Left Mouse Button"), FKeyDetails::MouseButton));
 	AddKey(FKeyDetails(EKeys::RightMouseButton, LOCTEXT("RightMouseButton", "Right Mouse Button"), FKeyDetails::MouseButton));
 	AddKey(FKeyDetails(EKeys::MiddleMouseButton, LOCTEXT("MiddleMouseButton", "Middle Mouse Button"), FKeyDetails::MouseButton));
@@ -957,6 +960,12 @@ bool FKey::IsBindableInBlueprints() const
 	return (KeyDetails.IsValid() ? KeyDetails->IsBindableInBlueprints() : false);
 }
 
+bool FKey::ShouldUpdateAxisWithoutSamples() const
+{
+	ConditionalLookupKeyDetails();
+	return (KeyDetails.IsValid() ? KeyDetails->ShouldUpdateAxisWithoutSamples() : false);
+}
+
 FText FKeyDetails::GetDisplayName() const
 {
 	return DisplayName.Get();
