Fixes to key menu categories.

diff --git a/Engine/Source/Runtime/InputCore/Private/InputCoreTypes.cpp b/Engine/Source/Runtime/InputCore/Private/InputCoreTypes.cpp
--- a/Engine/Source/Runtime/InputCore/Private/InputCoreTypes.cpp
+++ b/Engine/Source/Runtime/InputCore/Private/InputCoreTypes.cpp
@@ -575,7 +575,7 @@ void EKeys::Initialize()
 	AddKey(FKeyDetails(EKeys::Gesture_Flick, LOCTEXT("Gesture_Flick", "Flick"), 0, "Gesture"));
 
 	// Motion Controllers
-	AddMenuCategoryDisplayInfo("Steam", LOCTEXT("MotionControllerSubCategory", "Motion Controller"), TEXT("GraphEditor.PadEvent_16x"));
+	AddMenuCategoryDisplayInfo("MotionController", LOCTEXT("MotionControllerSubCategory", "Motion Controller"), TEXT("GraphEditor.PadEvent_16x"));
 
 	//	Left Controller
 	AddKey(FKeyDetails(EKeys::MotionController_Left_FaceButton1, LOCTEXT("MotionController_Left_FaceButton1", "MotionController (L) FaceButton1"), FKeyDetails::GamepadKey));
@@ -633,7 +633,7 @@ void EKeys::Initialize()
 	AddKey(FKeyDetails(EKeys::MotionController_Right_TriggerAxis, LOCTEXT("MotionController_Right_TriggerAxis", "MotionController (R) TriggerAxis"), FKeyDetails::GamepadKey | FKeyDetails::FloatAxis));
 
 	// PS4-specific
-	AddMenuCategoryDisplayInfo("Steam", LOCTEXT("PS4SubCategory", "PS4"), TEXT("GraphEditor.PadEvent_16x"));
+	AddMenuCategoryDisplayInfo("PS4", LOCTEXT("PS4SubCategory", "PS4"), TEXT("GraphEditor.PadEvent_16x"));
 	AddKey(FKeyDetails(EKeys::PS4_Special, LOCTEXT("PS4_Special", "PS4_Special"), FKeyDetails::NotBlueprintBindableKey));
 
 	// Steam Controller specific
