Allow binding of the Section (┬º) key

diff --git a/Engine/Source/Runtime/Core/Private/GenericPlatform/GenericPlatformMisc.cpp b/Engine/Source/Runtime/Core/Private/GenericPlatform/GenericPlatformMisc.cpp
--- a/Engine/Source/Runtime/Core/Private/GenericPlatform/GenericPlatformMisc.cpp
+++ b/Engine/Source/Runtime/Core/Private/GenericPlatform/GenericPlatformMisc.cpp
@@ -786,6 +786,7 @@ uint32 FGenericPlatformMisc::GetStandardPrintableKeyMap(uint32* KeyCodes, FStrin
 	ADDKEYMAP( 231, TEXT("C_Cedille") );
 	ADDKEYMAP( 233, TEXT("E_AccentAigu") );
 	ADDKEYMAP( 232, TEXT("E_AccentGrave") );
+	ADDKEYMAP( 167, TEXT("Section") );
 
 	return NumMappings;
 }
diff --git a/Engine/Source/Runtime/InputCore/Classes/InputCoreTypes.h b/Engine/Source/Runtime/InputCore/Classes/InputCoreTypes.h
--- a/Engine/Source/Runtime/InputCore/Classes/InputCoreTypes.h
+++ b/Engine/Source/Runtime/InputCore/Classes/InputCoreTypes.h
@@ -307,6 +307,7 @@ struct INPUTCORE_API EKeys
 	static const FKey E_AccentGrave;
 	static const FKey E_AccentAigu;
 	static const FKey C_Cedille;
+	static const FKey Section;
 
 	// Platform Keys
 	// These keys platform specific versions of keys that go by different names.
diff --git a/Engine/Source/Runtime/InputCore/Private/InputCoreTypes.cpp b/Engine/Source/Runtime/InputCore/Private/InputCoreTypes.cpp
--- a/Engine/Source/Runtime/InputCore/Private/InputCoreTypes.cpp
+++ b/Engine/Source/Runtime/InputCore/Private/InputCoreTypes.cpp
@@ -151,6 +151,7 @@ const FKey EKeys::A_AccentGrave("A_AccentGrave");
 const FKey EKeys::E_AccentGrave("E_AccentGrave");
 const FKey EKeys::E_AccentAigu("E_AccentAigu");
 const FKey EKeys::C_Cedille("C_Cedille");
+const FKey EKeys::Section("Section");
 
 
 // Setup platform specific keys
@@ -507,6 +508,7 @@ void EKeys::Initialize()
 	AddKey(FKeyDetails(EKeys::E_AccentGrave, FText::FromString(FString::Chr(232))));
 	AddKey(FKeyDetails(EKeys::E_AccentAigu, FText::FromString(FString::Chr(233))));
 	AddKey(FKeyDetails(EKeys::C_Cedille, FText::FromString(FString::Chr(231))));
+	AddKey(FKeyDetails(EKeys::Section, FText::FromString(FString::Chr(167))));
 
 
 	// Setup Gamepad keys
