#ue4 Minor tweaks and fixes to Move Controller code before 4.10 integrate.

-Crash fix when using Move controllers with Oculus HMD
-Changing "Move" button to trigger Grip1 instead of shoulder, this is more consistant with Touch controller usage.

diff --git a/Engine/Source/Runtime/Core/Private/GenericPlatform/GenericApplication.cpp b/Engine/Source/Runtime/Core/Private/GenericPlatform/GenericApplication.cpp
--- a/Engine/Source/Runtime/Core/Private/GenericPlatform/GenericApplication.cpp
+++ b/Engine/Source/Runtime/Core/Private/GenericPlatform/GenericApplication.cpp
@@ -89,11 +89,15 @@ const FGamepadKeyNames::Type FGamepadKeyNames::MotionController_Right_Thumbstick
 const FGamepadKeyNames::Type FGamepadKeyNames::MotionController_Left_Thumbstick_X("MotionController_Left_Thumbstick_X");
 const FGamepadKeyNames::Type FGamepadKeyNames::MotionController_Left_Thumbstick_Y("MotionController_Left_Thumbstick_Y");
 const FGamepadKeyNames::Type FGamepadKeyNames::MotionController_Left_TriggerAxis("MotionController_Left_TriggerAxis");
+const FGamepadKeyNames::Type FGamepadKeyNames::MotionController_Left_Grip1Axis( "MotionController_Left_Grip1Axis" );
+const FGamepadKeyNames::Type FGamepadKeyNames::MotionController_Left_Grip2Axis( "MotionController_Left_Grip2Axis" );
 
 	//		Right Controller
 const FGamepadKeyNames::Type FGamepadKeyNames::MotionController_Right_Thumbstick_X("MotionController_Right_Thumbstick_X");
 const FGamepadKeyNames::Type FGamepadKeyNames::MotionController_Right_Thumbstick_Y("MotionController_Right_Thumbstick_Y");
 const FGamepadKeyNames::Type FGamepadKeyNames::MotionController_Right_TriggerAxis("MotionController_Right_TriggerAxis");
+const FGamepadKeyNames::Type FGamepadKeyNames::MotionController_Right_Grip1Axis( "MotionController_Right_Grip1Axis" );
+const FGamepadKeyNames::Type FGamepadKeyNames::MotionController_Right_Grip2Axis( "MotionController_Right_Grip2Axis" );
 
 void FDisplayMetrics::PrintToLog() const
 {
diff --git a/Engine/Source/Runtime/Core/Public/GenericPlatform/GenericApplicationMessageHandler.h b/Engine/Source/Runtime/Core/Public/GenericPlatform/GenericApplicationMessageHandler.h
--- a/Engine/Source/Runtime/Core/Public/GenericPlatform/GenericApplicationMessageHandler.h
+++ b/Engine/Source/Runtime/Core/Public/GenericPlatform/GenericApplicationMessageHandler.h
@@ -114,11 +114,15 @@ struct CORE_API FGamepadKeyNames
 	static const FName MotionController_Left_Thumbstick_X;
 	static const FName MotionController_Left_Thumbstick_Y;
 	static const FName MotionController_Left_TriggerAxis;
+	static const FName MotionController_Left_Grip1Axis;
+	static const FName MotionController_Left_Grip2Axis;
 
 	//		Right Controller
 	static const FName MotionController_Right_Thumbstick_X;
 	static const FName MotionController_Right_Thumbstick_Y;
 	static const FName MotionController_Right_TriggerAxis;
+	static const FName MotionController_Right_Grip1Axis;
+	static const FName MotionController_Right_Grip2Axis;
 };
 
 namespace EWindowActivation
diff --git a/Engine/Source/Runtime/InputCore/Classes/InputCoreTypes.h b/Engine/Source/Runtime/InputCore/Classes/InputCoreTypes.h
--- a/Engine/Source/Runtime/InputCore/Classes/InputCoreTypes.h
+++ b/Engine/Source/Runtime/InputCore/Classes/InputCoreTypes.h
@@ -411,11 +411,15 @@ struct INPUTCORE_API EKeys
 	static const FKey MotionController_Left_Thumbstick_X;
 	static const FKey MotionController_Left_Thumbstick_Y;
 	static const FKey MotionController_Left_TriggerAxis;
+	static const FKey MotionController_Left_Grip1Axis;
+	static const FKey MotionController_Left_Grip2Axis;
 
 	//		Right Controller
 	static const FKey MotionController_Right_Thumbstick_X;
 	static const FKey MotionController_Right_Thumbstick_Y;
 	static const FKey MotionController_Right_TriggerAxis;
+	static const FKey MotionController_Right_Grip1Axis;
+	static const FKey MotionController_Right_Grip2Axis;
 
 	// PS4-specific
 	static const FKey PS4_Special;
diff --git a/Engine/Source/Runtime/InputCore/Private/InputCoreTypes.cpp b/Engine/Source/Runtime/InputCore/Private/InputCoreTypes.cpp
--- a/Engine/Source/Runtime/InputCore/Private/InputCoreTypes.cpp
+++ b/Engine/Source/Runtime/InputCore/Private/InputCoreTypes.cpp
@@ -272,11 +272,15 @@ const FKey EKeys::MotionController_Right_Thumbstick_Right("MotionController_Righ
 const FKey EKeys::MotionController_Left_Thumbstick_X("MotionController_Left_Thumbstick_X");
 const FKey EKeys::MotionController_Left_Thumbstick_Y("MotionController_Left_Thumbstick_Y");
 const FKey EKeys::MotionController_Left_TriggerAxis("MotionController_Left_TriggerAxis");
+const FKey EKeys::MotionController_Left_Grip1Axis( "MotionController_Left_Grip1Axis" );
+const FKey EKeys::MotionController_Left_Grip2Axis( "MotionController_Left_Grip2Axis" );
 
 //		Right Controller
 const FKey EKeys::MotionController_Right_Thumbstick_X("MotionController_Right_Thumbstick_X");
 const FKey EKeys::MotionController_Right_Thumbstick_Y("MotionController_Right_Thumbstick_Y");
 const FKey EKeys::MotionController_Right_TriggerAxis("MotionController_Right_TriggerAxis");
+const FKey EKeys::MotionController_Right_Grip1Axis( "MotionController_Right_Grip1Axis" );
+const FKey EKeys::MotionController_Right_Grip2Axis( "MotionController_Right_Grip2Axis" );
 
 // PS4-specific
 const FKey EKeys::PS4_Special("PS4_Special");
@@ -628,11 +632,15 @@ void EKeys::Initialize()
 	AddKey(FKeyDetails(EKeys::MotionController_Left_Thumbstick_X, LOCTEXT("MotionController_Left_Thumbstick_X", "MotionController (L) Thumbstick X"), FKeyDetails::GamepadKey | FKeyDetails::FloatAxis));
 	AddKey(FKeyDetails(EKeys::MotionController_Left_Thumbstick_Y, LOCTEXT("MotionController_Left_Thumbstick_Y", "MotionController (L) Thumbstick Y"), FKeyDetails::GamepadKey | FKeyDetails::FloatAxis));
 	AddKey(FKeyDetails(EKeys::MotionController_Left_TriggerAxis, LOCTEXT("MotionController_Left_TriggerAxis", "MotionController (L) TriggerAxis"), FKeyDetails::GamepadKey | FKeyDetails::FloatAxis));
+	AddKey(FKeyDetails(EKeys::MotionController_Left_Grip1Axis, LOCTEXT("MotionController_Left_Grip1Axis", "MotionController (L) Grip1 Axis"), FKeyDetails::GamepadKey | FKeyDetails::FloatAxis));
+	AddKey(FKeyDetails(EKeys::MotionController_Left_Grip2Axis, LOCTEXT("MotionController_Left_Grip2Axis", "MotionController (L) Grip2 Axis"), FKeyDetails::GamepadKey | FKeyDetails::FloatAxis));
 	
 	//	Right Controller
 	AddKey(FKeyDetails(EKeys::MotionController_Right_Thumbstick_X, LOCTEXT("MotionController_Right_Thumbstick_X", "MotionController (R) Thumbstick X"), FKeyDetails::GamepadKey | FKeyDetails::FloatAxis));
 	AddKey(FKeyDetails(EKeys::MotionController_Right_Thumbstick_Y, LOCTEXT("MotionController_Right_Thumbstick_Y", "MotionController (R) Thumbstick Y"), FKeyDetails::GamepadKey | FKeyDetails::FloatAxis));
 	AddKey(FKeyDetails(EKeys::MotionController_Right_TriggerAxis, LOCTEXT("MotionController_Right_TriggerAxis", "MotionController (R) TriggerAxis"), FKeyDetails::GamepadKey | FKeyDetails::FloatAxis));
+	AddKey(FKeyDetails(EKeys::MotionController_Right_Grip1Axis, LOCTEXT("MotionController_Right_Grip1Axis", "MotionController (R) Grip1 Axis"), FKeyDetails::GamepadKey | FKeyDetails::FloatAxis));
+	AddKey(FKeyDetails(EKeys::MotionController_Right_Grip2Axis, LOCTEXT("MotionController_Right_Grip2Axis", "MotionController (R) Grip2 Axis"), FKeyDetails::GamepadKey | FKeyDetails::FloatAxis));
 
 	// PS4-specific
 	AddMenuCategoryDisplayInfo("PS4", LOCTEXT("PS4SubCategory", "PS4"), TEXT("GraphEditor.PadEvent_16x"));
