#rnx

diff --git a/Engine/Source/Runtime/InputCore/Private/HTML5/HTML5PlatformInput.cpp b/Engine/Source/Runtime/InputCore/Private/HTML5/HTML5PlatformInput.cpp
--- a/Engine/Source/Runtime/InputCore/Private/HTML5/HTML5PlatformInput.cpp
+++ b/Engine/Source/Runtime/InputCore/Private/HTML5/HTML5PlatformInput.cpp
@@ -78,7 +78,8 @@ uint32 FHTML5PlatformInput::GetKeyMap(uint32* KeyCodes, FString* KeyNames, uint3
 		ADDKEYMAP(SDL_SCANCODE_CAPSLOCK, TEXT("CapsLock"));
 		ADDKEYMAP(SDL_SCANCODE_NUMLOCKCLEAR, TEXT("NumLock"));
 		ADDKEYMAP(SDL_SCANCODE_SCROLLLOCK, TEXT("ScrollLock"));
-		ADDKEYMAP(SDL_SCANCODE_CANCEL, TEXT("Cancel")); // CTRL+Pause/Break and CTRL+ScrollLock -- UE-58441
+		ADDKEYMAP(SDL_SCANCODE_CANCEL, TEXT("Cancel")); // UE-58441 -- CTRL+Pause/Break and CTRL+ScrollLock
+		ADDKEYMAP(SDL_SCANCODE_APPLICATION, TEXT("Menu")); // UE-58441 -- contextual menu
 	}
 #undef ADDKEYMAP
 
