Reset the KeyDetails on an FKey when assigned in blueprints.

diff --git a/Engine/Source/Runtime/InputCore/Classes/InputCoreTypes.h b/Engine/Source/Runtime/InputCore/Classes/InputCoreTypes.h
--- a/Engine/Source/Runtime/InputCore/Classes/InputCoreTypes.h
+++ b/Engine/Source/Runtime/InputCore/Classes/InputCoreTypes.h
@@ -79,6 +79,7 @@ struct INPUTCORE_API FKey
 	bool ExportTextItem(FString& ValueStr, FKey const& DefaultValue, UObject* Parent, int32 PortFlags, UObject* ExportRootScope) const;
 	bool ImportTextItem(const TCHAR*& Buffer, int32 PortFlags, UObject* Parent, FOutputDevice* ErrorText);
 	void PostSerialize(const FArchive& Ar);
+	void PostScriptConstruct();
 
 	friend bool operator==(const FKey& KeyA, const FKey& KeyB) { return KeyA.KeyName == KeyB.KeyName; }
 	friend bool operator!=(const FKey& KeyA, const FKey& KeyB) { return KeyA.KeyName != KeyB.KeyName; }
@@ -109,6 +110,7 @@ struct TStructOpsTypeTraits<FKey> : public TStructOpsTypeTraitsBase2<FKey>
 		WithExportTextItem = true,
 		WithImportTextItem = true,
 		WithPostSerialize = true,
+		WithPostScriptConstruct = true,
 		WithCopy = true,		// Necessary so that TSharedPtr<FKeyDetails> Data is copied around
 	};
 };
diff --git a/Engine/Source/Runtime/InputCore/Private/InputCoreTypes.cpp b/Engine/Source/Runtime/InputCore/Private/InputCoreTypes.cpp
--- a/Engine/Source/Runtime/InputCore/Private/InputCoreTypes.cpp
+++ b/Engine/Source/Runtime/InputCore/Private/InputCoreTypes.cpp
@@ -1102,6 +1102,11 @@ void FKey::PostSerialize(const FArchive& Ar)
 	ResetKey();
 }
 
+void FKey::PostScriptConstruct()
+{
+	KeyDetails.Reset();
+}
+
 void FKey::ResetKey()
 {
 	KeyDetails.Reset();
