Fix ensures when using synthetic keys

diff --git a/Engine/Source/Runtime/InputCore/Private/InputCoreTypes.cpp b/Engine/Source/Runtime/InputCore/Private/InputCoreTypes.cpp
--- a/Engine/Source/Runtime/InputCore/Private/InputCoreTypes.cpp
+++ b/Engine/Source/Runtime/InputCore/Private/InputCoreTypes.cpp
@@ -1188,6 +1188,7 @@ FKey FInputKeyManager::GetKeyFromCodes( const uint32 KeyCode, const uint32 CharC
 	{
 		FKey NewKey(*FString::Printf(TEXT("%s%d"), FKey::SyntheticCharPrefix, CharCode));
 		EKeys::AddKey(FKeyDetails(NewKey, FText::AsCultureInvariant(FString::Chr(CharCode)), FKeyDetails::NotBlueprintBindableKey));
+		const_cast<FInputKeyManager*>(this)->KeyMapCharToEnum.Add(CharCode, NewKey);
 		return NewKey;
 	}
 	return KeyPtr ? *KeyPtr : EKeys::Invalid;
