Deprecate MotionController keys

diff --git a/Engine/Source/Runtime/InputCore/Private/InputCoreTypes.cpp b/Engine/Source/Runtime/InputCore/Private/InputCoreTypes.cpp
--- a/Engine/Source/Runtime/InputCore/Private/InputCoreTypes.cpp
+++ b/Engine/Source/Runtime/InputCore/Private/InputCoreTypes.cpp
@@ -275,6 +275,7 @@ const FKey EKeys::MotionController_Right_Thumbstick_Right("MotionController_Righ
 
 //   Motion Controller Axes
 //		Left Controller
+const FKey EKeys::MotionController_Left_Thumbstick_XY("MotionController_Left_Thumbstick_XY");
 const FKey EKeys::MotionController_Left_Thumbstick_X("MotionController_Left_Thumbstick_X");
 const FKey EKeys::MotionController_Left_Thumbstick_Y("MotionController_Left_Thumbstick_Y");
 const FKey EKeys::MotionController_Left_TriggerAxis("MotionController_Left_TriggerAxis");
@@ -282,6 +283,7 @@ const FKey EKeys::MotionController_Left_Grip1Axis( "MotionController_Left_Grip1A
 const FKey EKeys::MotionController_Left_Grip2Axis( "MotionController_Left_Grip2Axis" );
 
 //		Right Controller
+const FKey EKeys::MotionController_Right_Thumbstick_XY("MotionController_Right_Thumbstick_XY");
 const FKey EKeys::MotionController_Right_Thumbstick_X("MotionController_Right_Thumbstick_X");
 const FKey EKeys::MotionController_Right_Thumbstick_Y("MotionController_Right_Thumbstick_Y");
 const FKey EKeys::MotionController_Right_TriggerAxis("MotionController_Right_TriggerAxis");
@@ -312,6 +314,152 @@ const FKey EKeys::Android_Volume_Up("Android_Volume_Up");
 const FKey EKeys::Android_Volume_Down("Android_Volume_Down");
 const FKey EKeys::Android_Menu("Android_Menu");
 
+// Google Daydream
+const FKey EKeys::Daydream_Left_Select_Click("Daydream_Left_Select_Click");
+const FKey EKeys::Daydream_Left_Trackpad_X("Daydream_Left_Trackpad_X");
+const FKey EKeys::Daydream_Left_Trackpad_Y("Daydream_Left_Trackpad_Y");
+const FKey EKeys::Daydream_Left_Trackpad_Click("Daydream_Left_Trackpad_Click");
+const FKey EKeys::Daydream_Left_Trackpad_Touch("Daydream_Left_Trackpad_Touch");
+const FKey EKeys::Daydream_Right_Select_Click("Daydream_Right_Select_Click");
+const FKey EKeys::Daydream_Right_Trackpad_X("Daydream_Right_Trackpad_X");
+const FKey EKeys::Daydream_Right_Trackpad_Y("Daydream_Right_Trackpad_Y");
+const FKey EKeys::Daydream_Right_Trackpad_Click("Daydream_Right_Trackpad_Click");
+const FKey EKeys::Daydream_Right_Trackpad_Touch("Daydream_Right_Trackpad_Touch");
+
+// HTC Vive Controller
+const FKey EKeys::Vive_Left_System_Click("Vive_Left_System_Click");
+const FKey EKeys::Vive_Left_Grip_Click("Vive_Left_Grip_Click");
+const FKey EKeys::Vive_Left_Menu_Click("Vive_Left_Menu_Click");
+const FKey EKeys::Vive_Left_Trigger_Click("Vive_Left_Trigger_Click");
+const FKey EKeys::Vive_Left_Trigger_Axis("Vive_Left_Trigger_Axis");
+const FKey EKeys::Vive_Left_Trackpad_X("Vive_Left_Trackpad_X");
+const FKey EKeys::Vive_Left_Trackpad_Y("Vive_Left_Trackpad_Y");
+const FKey EKeys::Vive_Left_Trackpad_Click("Vive_Left_Trackpad_Click");
+const FKey EKeys::Vive_Left_Trackpad_Touch("Vive_Left_Trackpad_Touch");
+const FKey EKeys::Vive_Right_System_Click("Vive_Right_System_Click");
+const FKey EKeys::Vive_Right_Grip_Click("Vive_Right_Grip_Click");
+const FKey EKeys::Vive_Right_Menu_Click("Vive_Right_Menu_Click");
+const FKey EKeys::Vive_Right_Trigger_Click("Vive_Right_Trigger_Click");
+const FKey EKeys::Vive_Right_Trigger_Axis("Vive_Right_Trigger_Axis");
+const FKey EKeys::Vive_Right_Trackpad_X("Vive_Right_Trackpad_X");
+const FKey EKeys::Vive_Right_Trackpad_Y("Vive_Right_Trackpad_Y");
+const FKey EKeys::Vive_Right_Trackpad_Click("Vive_Right_Trackpad_Click");
+const FKey EKeys::Vive_Right_Trackpad_Touch("Vive_Right_Trackpad_Touch");
+
+// Microsoft Mixed Reality Motion Controller
+const FKey EKeys::MixedReality_Left_Menu_Click("MixedReality_Left_Menu_Click");
+const FKey EKeys::MixedReality_Left_Grip_Click("MixedReality_Left_Grip_Click");
+const FKey EKeys::MixedReality_Left_Trigger_Click("MixedReality_Left_Trigger_Click");
+const FKey EKeys::MixedReality_Left_Trigger_Axis("MixedReality_Left_Trigger_Axis");
+const FKey EKeys::MixedReality_Left_Thumbstick_X("MixedReality_Left_Thumbstick_X");
+const FKey EKeys::MixedReality_Left_Thumbstick_Y("MixedReality_Left_Thumbstick_Y");
+const FKey EKeys::MixedReality_Left_Thumbstick_Click("MixedReality_Left_Thumbstick_Click");
+const FKey EKeys::MixedReality_Left_Trackpad_X("MixedReality_Left_Trackpad_X");
+const FKey EKeys::MixedReality_Left_Trackpad_Y("MixedReality_Left_Trackpad_Y");
+const FKey EKeys::MixedReality_Left_Trackpad_Click("MixedReality_Left_Trackpad_Click");
+const FKey EKeys::MixedReality_Left_Trackpad_Touch("MixedReality_Left_Trackpad_Touch");
+const FKey EKeys::MixedReality_Right_Menu_Click("MixedReality_Right_Menu_Click");
+const FKey EKeys::MixedReality_Right_Grip_Click("MixedReality_Right_Grip_Click");
+const FKey EKeys::MixedReality_Right_Trigger_Click("MixedReality_Right_Trigger_Click");
+const FKey EKeys::MixedReality_Right_Trigger_Axis("MixedReality_Right_Trigger_Axis");
+const FKey EKeys::MixedReality_Right_Thumbstick_X("MixedReality_Right_Thumbstick_X");
+const FKey EKeys::MixedReality_Right_Thumbstick_Y("MixedReality_Right_Thumbstick_Y");
+const FKey EKeys::MixedReality_Right_Thumbstick_Click("MixedReality_Right_Thumbstick_Click");
+const FKey EKeys::MixedReality_Right_Trackpad_X("MixedReality_Right_Trackpad_X");
+const FKey EKeys::MixedReality_Right_Trackpad_Y("MixedReality_Right_Trackpad_Y");
+const FKey EKeys::MixedReality_Right_Trackpad_Click("MixedReality_Right_Trackpad_Click");
+const FKey EKeys::MixedReality_Right_Trackpad_Touch("MixedReality_Right_Trackpad_Touch");
+
+// Oculus Go Controller
+const FKey EKeys::OculusGo_Left_System_Click("OculusGo_Left_System_Click");
+const FKey EKeys::OculusGo_Left_Back_Click("OculusGo_Left_Back_Click");
+const FKey EKeys::OculusGo_Left_Trigger_Click("OculusGo_Left_Trigger_Click");
+const FKey EKeys::OculusGo_Left_Trackpad_X("OculusGo_Left_Trackpad_X");
+const FKey EKeys::OculusGo_Left_Trackpad_Y("OculusGo_Left_Trackpad_Y");
+const FKey EKeys::OculusGo_Left_Trackpad_Click("OculusGo_Left_Trackpad_Click");
+const FKey EKeys::OculusGo_Left_Trackpad_Touch("OculusGo_Left_Trackpad_Touch");
+const FKey EKeys::OculusGo_Right_System_Click("OculusGo_Right_System_Click");
+const FKey EKeys::OculusGo_Right_Back_Click("OculusGo_Right_Back_Click");
+const FKey EKeys::OculusGo_Right_Trigger_Click("OculusGo_Right_Trigger_Click");
+const FKey EKeys::OculusGo_Right_Trackpad_X("OculusGo_Right_Trackpad_X");
+const FKey EKeys::OculusGo_Right_Trackpad_Y("OculusGo_Right_Trackpad_Y");
+const FKey EKeys::OculusGo_Right_Trackpad_Click("OculusGo_Right_Trackpad_Click");
+const FKey EKeys::OculusGo_Right_Trackpad_Touch("OculusGo_Right_Trackpad_Touch");
+
+// Oculus Touch Controller
+const FKey EKeys::OculusTouch_Left_X_Click("OculusTouch_Left_X_Click");
+const FKey EKeys::OculusTouch_Left_Y_Click("OculusTouch_Left_Y_Click");
+const FKey EKeys::OculusTouch_Left_X_Touch("OculusTouch_Left_X_Touch");
+const FKey EKeys::OculusTouch_Left_Y_Touch("OculusTouch_Left_Y_Touch");
+const FKey EKeys::OculusTouch_Left_Menu_Click("OculusTouch_Left_Menu_Click");
+const FKey EKeys::OculusTouch_Left_Grip_Click("OculusTouch_Left_Grip_Click");
+const FKey EKeys::OculusTouch_Left_Grip_Axis("OculusTouch_Left_Grip_Axis");
+const FKey EKeys::OculusTouch_Left_Trigger_Click("OculusTouch_Left_Trigger_Click");
+const FKey EKeys::OculusTouch_Left_Trigger_Axis("OculusTouch_Left_Trigger_Axis");
+const FKey EKeys::OculusTouch_Left_Trigger_Touch("OculusTouch_Left_Trigger_Touch");
+const FKey EKeys::OculusTouch_Left_Thumbstick_X("OculusTouch_Left_Thumbstick_X");
+const FKey EKeys::OculusTouch_Left_Thumbstick_Y("OculusTouch_Left_Thumbstick_Y");
+const FKey EKeys::OculusTouch_Left_Thumbstick_Click("OculusTouch_Left_Thumbstick_Click");
+const FKey EKeys::OculusTouch_Left_Thumbstick_Touch("OculusTouch_Left_Thumbstick_Touch");
+const FKey EKeys::OculusTouch_Right_A_Click("OculusTouch_Right_A_Click");
+const FKey EKeys::OculusTouch_Right_B_Click("OculusTouch_Right_B_Click");
+const FKey EKeys::OculusTouch_Right_A_Touch("OculusTouch_Right_A_Touch");
+const FKey EKeys::OculusTouch_Right_B_Touch("OculusTouch_Right_B_Touch");
+const FKey EKeys::OculusTouch_Right_System_Click("OculusTouch_Right_System_Click");
+const FKey EKeys::OculusTouch_Right_Grip_Click("OculusTouch_Right_Grip_Click");
+const FKey EKeys::OculusTouch_Right_Grip_Axis("OculusTouch_Right_Grip_Axis");
+const FKey EKeys::OculusTouch_Right_Trigger_Click("OculusTouch_Right_Trigger_Click");
+const FKey EKeys::OculusTouch_Right_Trigger_Axis("OculusTouch_Right_Trigger_Axis");
+const FKey EKeys::OculusTouch_Right_Trigger_Touch("OculusTouch_Right_Trigger_Touch");
+const FKey EKeys::OculusTouch_Right_Thumbstick_X("OculusTouch_Right_Thumbstick_X");
+const FKey EKeys::OculusTouch_Right_Thumbstick_Y("OculusTouch_Right_Thumbstick_Y");
+const FKey EKeys::OculusTouch_Right_Thumbstick_Click("OculusTouch_Right_Thumbstick_Click");
+const FKey EKeys::OculusTouch_Right_Thumbstick_Touch("OculusTouch_Right_Thumbstick_Touch");
+
+// Valve Index Controller
+const FKey EKeys::ValveIndex_Left_A_Click("ValveIndex_Left_A_Click");
+const FKey EKeys::ValveIndex_Left_B_Click("ValveIndex_Left_B_Click");
+const FKey EKeys::ValveIndex_Left_A_Touch("ValveIndex_Left_A_Touch");
+const FKey EKeys::ValveIndex_Left_B_Touch("ValveIndex_Left_B_Touch");
+const FKey EKeys::ValveIndex_Left_System_Click("ValveIndex_Left_System_Click");
+const FKey EKeys::ValveIndex_Left_System_Touch("ValveIndex_Left_System_Touch");
+const FKey EKeys::ValveIndex_Left_Grip_Click("ValveIndex_Left_Grip_Click");
+const FKey EKeys::ValveIndex_Left_Grip_Axis("ValveIndex_Left_Grip_Axis");
+const FKey EKeys::ValveIndex_Left_Grip_Force("ValveIndex_Left_Grip_Force");
+const FKey EKeys::ValveIndex_Left_Trigger_Click("ValveIndex_Left_Trigger_Click");
+const FKey EKeys::ValveIndex_Left_Trigger_Axis("ValveIndex_Left_Trigger_Axis");
+const FKey EKeys::ValveIndex_Left_Trigger_Touch("ValveIndex_Left_Trigger_Touch");
+const FKey EKeys::ValveIndex_Left_Thumbstick_X("ValveIndex_Left_Thumbstick_X");
+const FKey EKeys::ValveIndex_Left_Thumbstick_Y("ValveIndex_Left_Thumbstick_Y");
+const FKey EKeys::ValveIndex_Left_Thumbstick_Click("ValveIndex_Left_Thumbstick_Click");
+const FKey EKeys::ValveIndex_Left_Thumbstick_Touch("ValveIndex_Left_Thumbstick_Touch");
+const FKey EKeys::ValveIndex_Left_Trackpad_X("ValveIndex_Left_Trackpad_X");
+const FKey EKeys::ValveIndex_Left_Trackpad_Y("ValveIndex_Left_Trackpad_Y");
+const FKey EKeys::ValveIndex_Left_Trackpad_Click("ValveIndex_Left_Trackpad_Click");
+const FKey EKeys::ValveIndex_Left_Trackpad_Force("ValveIndex_Left_Trackpad_Force");
+const FKey EKeys::ValveIndex_Left_Trackpad_Touch("ValveIndex_Left_Trackpad_Touch");
+const FKey EKeys::ValveIndex_Right_A_Click("ValveIndex_Right_A_Click");
+const FKey EKeys::ValveIndex_Right_B_Click("ValveIndex_Right_B_Click");
+const FKey EKeys::ValveIndex_Right_A_Touch("ValveIndex_Right_A_Touch");
+const FKey EKeys::ValveIndex_Right_B_Touch("ValveIndex_Right_B_Touch");
+const FKey EKeys::ValveIndex_Right_System_Click("ValveIndex_Right_System_Click");
+const FKey EKeys::ValveIndex_Right_System_Touch("ValveIndex_Right_System_Touch");
+const FKey EKeys::ValveIndex_Right_Grip_Click("ValveIndex_Right_Grip_Click");
+const FKey EKeys::ValveIndex_Right_Grip_Axis("ValveIndex_Right_Grip_Axis");
+const FKey EKeys::ValveIndex_Right_Grip_Force("ValveIndex_Right_Grip_Force");
+const FKey EKeys::ValveIndex_Right_Trigger_Click("ValveIndex_Right_Trigger_Click");
+const FKey EKeys::ValveIndex_Right_Trigger_Axis("ValveIndex_Right_Trigger_Axis");
+const FKey EKeys::ValveIndex_Right_Trigger_Touch("ValveIndex_Right_Trigger_Touch");
+const FKey EKeys::ValveIndex_Right_Thumbstick_X("ValveIndex_Right_Thumbstick_X");
+const FKey EKeys::ValveIndex_Right_Thumbstick_Y("ValveIndex_Right_Thumbstick_Y");
+const FKey EKeys::ValveIndex_Right_Thumbstick_Click("ValveIndex_Right_Thumbstick_Click");
+const FKey EKeys::ValveIndex_Right_Thumbstick_Touch("ValveIndex_Right_Thumbstick_Touch");
+const FKey EKeys::ValveIndex_Right_Trackpad_X("ValveIndex_Right_Trackpad_X");
+const FKey EKeys::ValveIndex_Right_Trackpad_Y("ValveIndex_Right_Trackpad_Y");
+const FKey EKeys::ValveIndex_Right_Trackpad_Click("ValveIndex_Right_Trackpad_Click");
+const FKey EKeys::ValveIndex_Right_Trackpad_Force("ValveIndex_Right_Trackpad_Force");
+const FKey EKeys::ValveIndex_Right_Trackpad_Touch("ValveIndex_Right_Trackpad_Touch");
+
 const FKey EKeys::Invalid(NAME_None);
 
 const FName EKeys::NAME_GamepadCategory("Gamepad");
@@ -325,7 +473,7 @@ bool EKeys::bInitialized = false;
 TMap<FKey, TSharedPtr<FKeyDetails> > EKeys::InputKeys;
 TMap<FName, EKeys::FCategoryDisplayInfo> EKeys::MenuCategoryDisplayInfo;
 
-FKeyDetails::FKeyDetails(const FKey InKey, const TAttribute<FText>& InLongDisplayName, const TAttribute<FText>& InShortDisplayName, const uint8 InKeyFlags, const FName InMenuCategory)
+FKeyDetails::FKeyDetails(const FKey InKey, const TAttribute<FText>& InLongDisplayName, const TAttribute<FText>& InShortDisplayName, const uint16 InKeyFlags, const FName InMenuCategory)
 	: Key(InKey)
 	, MenuCategory(InMenuCategory)
 	, LongDisplayName(InLongDisplayName)
@@ -334,7 +482,7 @@ FKeyDetails::FKeyDetails(const FKey InKey, const TAttribute<FText>& InLongDispla
 	CommonInit(InKeyFlags);
 }
 
-FKeyDetails::FKeyDetails(const FKey InKey, const TAttribute<FText>& InLongDisplayName, const uint8 InKeyFlags, const FName InMenuCategory, const TAttribute<FText>& InShortDisplayName)
+FKeyDetails::FKeyDetails(const FKey InKey, const TAttribute<FText>& InLongDisplayName, const uint16 InKeyFlags, const FName InMenuCategory, const TAttribute<FText>& InShortDisplayName)
 	: Key(InKey)
 	, MenuCategory(InMenuCategory)
 	, LongDisplayName(InLongDisplayName)
@@ -343,7 +491,7 @@ FKeyDetails::FKeyDetails(const FKey InKey, const TAttribute<FText>& InLongDispla
 	CommonInit(InKeyFlags);
 }
 
-void FKeyDetails::CommonInit(const uint8 InKeyFlags)
+void FKeyDetails::CommonInit(const uint16 InKeyFlags)
 {
 	bIsModifierKey = ((InKeyFlags & EKeyFlags::ModifierKey) != 0);
 	bIsGamepadKey = ((InKeyFlags & EKeyFlags::GamepadKey) != 0);
@@ -351,7 +499,8 @@ void FKeyDetails::CommonInit(const uint8 InKeyFlags)
 	bIsMouseButton = ((InKeyFlags & EKeyFlags::MouseButton) != 0);
 	bIsBindableInBlueprints = ((~InKeyFlags & EKeyFlags::NotBlueprintBindableKey) != 0);
 	bShouldUpdateAxisWithoutSamples = ((InKeyFlags & EKeyFlags::UpdateAxisWithoutSamples) != 0);
-	
+	bIsBindableToActions = ((~InKeyFlags & EKeyFlags::NotActionBindableKey) != 0);
+
 	if ((InKeyFlags & EKeyFlags::FloatAxis) != 0)
 	{
 		ensure((InKeyFlags & EKeyFlags::VectorAxis) == 0);
@@ -391,7 +540,7 @@ void EKeys::Initialize()
 {
 	if (bInitialized) return;
 	bInitialized = true;
-	
+
 	AddMenuCategoryDisplayInfo(NAME_GamepadCategory, LOCTEXT("GamepadSubCategory", "Gamepad"), TEXT("GraphEditor.PadEvent_16x"));
 	AddMenuCategoryDisplayInfo(NAME_MouseCategory, LOCTEXT("MouseSubCategory", "Mouse"), TEXT("GraphEditor.MouseEvent_16x"));
 	AddMenuCategoryDisplayInfo(NAME_KeyboardCategory, LOCTEXT("KeyboardSubCategory", "Keyboard"), TEXT("GraphEditor.KeyEvent_16x"));
@@ -403,13 +552,13 @@ void EKeys::Initialize()
 	AddKey(FKeyDetails(EKeys::MouseWheelAxis, LOCTEXT("MouseWheelAxis", "Mouse Wheel Axis"), FKeyDetails::FloatAxis | FKeyDetails::MouseButton | FKeyDetails::UpdateAxisWithoutSamples));
 	AddKey(FKeyDetails(EKeys::MouseScrollUp, LOCTEXT("MouseScrollUp", "Mouse Wheel Up"), FKeyDetails::MouseButton));
 	AddKey(FKeyDetails(EKeys::MouseScrollDown, LOCTEXT("MouseScrollDown", "Mouse Wheel Down"), FKeyDetails::MouseButton));
-	
+
 	AddKey(FKeyDetails(EKeys::LeftMouseButton, LOCTEXT("LeftMouseButton", "Left Mouse Button"), FKeyDetails::MouseButton));
 	AddKey(FKeyDetails(EKeys::RightMouseButton, LOCTEXT("RightMouseButton", "Right Mouse Button"), FKeyDetails::MouseButton));
 	AddKey(FKeyDetails(EKeys::MiddleMouseButton, LOCTEXT("MiddleMouseButton", "Middle Mouse Button"), FKeyDetails::MouseButton));
 	AddKey(FKeyDetails(EKeys::ThumbMouseButton, LOCTEXT("ThumbMouseButton", "Thumb Mouse Button"), FKeyDetails::MouseButton));
 	AddKey(FKeyDetails(EKeys::ThumbMouseButton2, LOCTEXT("ThumbMouseButton2", "Thumb Mouse Button 2"), FKeyDetails::MouseButton));
-    
+
 	AddKey(FKeyDetails(EKeys::Tab, LOCTEXT("Tab", "Tab")));
 	AddKey(FKeyDetails(EKeys::Enter, LOCTEXT("Enter", "Enter")));
 	AddKey(FKeyDetails(EKeys::Pause, LOCTEXT("Pause", "Pause")));
@@ -428,7 +577,7 @@ void EKeys::Initialize()
 	AddKey(FKeyDetails(EKeys::Down, LOCTEXT("Down", "Down")));
 
 	AddKey(FKeyDetails(EKeys::Insert, LOCTEXT("Insert", "Insert"), LOCTEXT("InsertShort", "Ins")));
-    
+
 #if PLATFORM_MAC
     AddKey(FKeyDetails(EKeys::BackSpace, LOCTEXT("Delete", "Delete"), LOCTEXT("DeleteShort", "Del")));
     AddKey(FKeyDetails(EKeys::Delete, LOCTEXT("ForwardDelete", "Fn+Delete")));
@@ -614,70 +763,68 @@ void EKeys::Initialize()
 	AddKey(FKeyDetails(EKeys::Gesture_Rotate, LOCTEXT("Gesture_Rotate", "Rotate"), 0, "Gesture"));
 
 	// Motion Controllers
-	AddMenuCategoryDisplayInfo("MotionController", LOCTEXT("MotionControllerSubCategory", "Motion Controller"), TEXT("GraphEditor.PadEvent_16x"));
-
 	//	Left Controller
-	AddKey(FKeyDetails(EKeys::MotionController_Left_FaceButton1, LOCTEXT("MotionController_Left_FaceButton1", "MotionController (L) FaceButton1"), FKeyDetails::GamepadKey));
-	AddKey(FKeyDetails(EKeys::MotionController_Left_FaceButton2, LOCTEXT("MotionController_Left_FaceButton2", "MotionController (L) FaceButton2"), FKeyDetails::GamepadKey));
-	AddKey(FKeyDetails(EKeys::MotionController_Left_FaceButton3, LOCTEXT("MotionController_Left_FaceButton3", "MotionController (L) FaceButton3"), FKeyDetails::GamepadKey));
-	AddKey(FKeyDetails(EKeys::MotionController_Left_FaceButton4, LOCTEXT("MotionController_Left_FaceButton4", "MotionController (L) FaceButton4"), FKeyDetails::GamepadKey));
-	AddKey(FKeyDetails(EKeys::MotionController_Left_FaceButton5, LOCTEXT("MotionController_Left_FaceButton5", "MotionController (L) FaceButton5"), FKeyDetails::GamepadKey));
-	AddKey(FKeyDetails(EKeys::MotionController_Left_FaceButton6, LOCTEXT("MotionController_Left_FaceButton6", "MotionController (L) FaceButton6"), FKeyDetails::GamepadKey));
-	AddKey(FKeyDetails(EKeys::MotionController_Left_FaceButton7, LOCTEXT("MotionController_Left_FaceButton7", "MotionController (L) FaceButton7"), FKeyDetails::GamepadKey));
-	AddKey(FKeyDetails(EKeys::MotionController_Left_FaceButton8, LOCTEXT("MotionController_Left_FaceButton8", "MotionController (L) FaceButton8"), FKeyDetails::GamepadKey));
-	
-	AddKey(FKeyDetails(EKeys::MotionController_Left_Shoulder, LOCTEXT("MotionController_Left_Shoulder", "MotionController (L) Shoulder"), FKeyDetails::GamepadKey));
-	AddKey(FKeyDetails(EKeys::MotionController_Left_Trigger, LOCTEXT("MotionController_Left_Trigger", "MotionController (L) Trigger"), FKeyDetails::GamepadKey));
-	
-	AddKey(FKeyDetails(EKeys::MotionController_Left_Grip1, LOCTEXT("MotionController_Left_Grip1", "MotionController (L) Grip1"), FKeyDetails::GamepadKey));
-	AddKey(FKeyDetails(EKeys::MotionController_Left_Grip2, LOCTEXT("MotionController_Left_Grip2", "MotionController (L) Grip2"), FKeyDetails::GamepadKey));
-	
-	AddKey(FKeyDetails(EKeys::MotionController_Left_Thumbstick, LOCTEXT("MotionController_Left_Thumbstick", "MotionController (L) Thumbstick"), FKeyDetails::GamepadKey));
-	AddKey(FKeyDetails(EKeys::MotionController_Left_Thumbstick_Up, LOCTEXT("MotionController_Left_Thumbstick_Up", "MotionController (L) Thumbstick Up"), FKeyDetails::GamepadKey));
-	AddKey(FKeyDetails(EKeys::MotionController_Left_Thumbstick_Down, LOCTEXT("MotionController_Left_Thumbstick_Down", "MotionController (L) Thumbstick Down"), FKeyDetails::GamepadKey));
-	AddKey(FKeyDetails(EKeys::MotionController_Left_Thumbstick_Left, LOCTEXT("MotionController_Left_Thumbstick_Left", "MotionController (L) Thumbstick Left"), FKeyDetails::GamepadKey));
-	AddKey(FKeyDetails(EKeys::MotionController_Left_Thumbstick_Right, LOCTEXT("MotionController_Left_Thumbstick_Right", "MotionController (L) Thumbstick Right"), FKeyDetails::GamepadKey));
-	
+	AddKey(FKeyDetails(EKeys::MotionController_Left_FaceButton1, LOCTEXT("MotionController_Left_FaceButton1", "MotionController (L) FaceButton1"), FKeyDetails::GamepadKey | FKeyDetails::NotBlueprintBindableKey | FKeyDetails::NotActionBindableKey));
+	AddKey(FKeyDetails(EKeys::MotionController_Left_FaceButton2, LOCTEXT("MotionController_Left_FaceButton2", "MotionController (L) FaceButton2"), FKeyDetails::GamepadKey | FKeyDetails::NotBlueprintBindableKey | FKeyDetails::NotActionBindableKey));
+	AddKey(FKeyDetails(EKeys::MotionController_Left_FaceButton3, LOCTEXT("MotionController_Left_FaceButton3", "MotionController (L) FaceButton3"), FKeyDetails::GamepadKey | FKeyDetails::NotBlueprintBindableKey | FKeyDetails::NotActionBindableKey));
+	AddKey(FKeyDetails(EKeys::MotionController_Left_FaceButton4, LOCTEXT("MotionController_Left_FaceButton4", "MotionController (L) FaceButton4"), FKeyDetails::GamepadKey | FKeyDetails::NotBlueprintBindableKey | FKeyDetails::NotActionBindableKey));
+	AddKey(FKeyDetails(EKeys::MotionController_Left_FaceButton5, LOCTEXT("MotionController_Left_FaceButton5", "MotionController (L) FaceButton5"), FKeyDetails::GamepadKey | FKeyDetails::NotBlueprintBindableKey | FKeyDetails::NotActionBindableKey));
+	AddKey(FKeyDetails(EKeys::MotionController_Left_FaceButton6, LOCTEXT("MotionController_Left_FaceButton6", "MotionController (L) FaceButton6"), FKeyDetails::GamepadKey | FKeyDetails::NotBlueprintBindableKey | FKeyDetails::NotActionBindableKey));
+	AddKey(FKeyDetails(EKeys::MotionController_Left_FaceButton7, LOCTEXT("MotionController_Left_FaceButton7", "MotionController (L) FaceButton7"), FKeyDetails::GamepadKey | FKeyDetails::NotBlueprintBindableKey | FKeyDetails::NotActionBindableKey));
+	AddKey(FKeyDetails(EKeys::MotionController_Left_FaceButton8, LOCTEXT("MotionController_Left_FaceButton8", "MotionController (L) FaceButton8"), FKeyDetails::GamepadKey | FKeyDetails::NotBlueprintBindableKey | FKeyDetails::NotActionBindableKey));
+
+	AddKey(FKeyDetails(EKeys::MotionController_Left_Shoulder, LOCTEXT("MotionController_Left_Shoulder", "MotionController (L) Shoulder"), FKeyDetails::GamepadKey | FKeyDetails::NotBlueprintBindableKey | FKeyDetails::NotActionBindableKey));
+	AddKey(FKeyDetails(EKeys::MotionController_Left_Trigger, LOCTEXT("MotionController_Left_Trigger", "MotionController (L) Trigger"), FKeyDetails::GamepadKey | FKeyDetails::NotBlueprintBindableKey | FKeyDetails::NotActionBindableKey));
+
+	AddKey(FKeyDetails(EKeys::MotionController_Left_Grip1, LOCTEXT("MotionController_Left_Grip1", "MotionController (L) Grip1"), FKeyDetails::GamepadKey | FKeyDetails::NotBlueprintBindableKey | FKeyDetails::NotActionBindableKey));
+	AddKey(FKeyDetails(EKeys::MotionController_Left_Grip2, LOCTEXT("MotionController_Left_Grip2", "MotionController (L) Grip2"), FKeyDetails::GamepadKey | FKeyDetails::NotBlueprintBindableKey | FKeyDetails::NotActionBindableKey));
+
+	AddKey(FKeyDetails(EKeys::MotionController_Left_Thumbstick, LOCTEXT("MotionController_Left_Thumbstick", "MotionController (L) Thumbstick"), FKeyDetails::GamepadKey | FKeyDetails::NotBlueprintBindableKey | FKeyDetails::NotActionBindableKey));
+	AddKey(FKeyDetails(EKeys::MotionController_Left_Thumbstick_Up, LOCTEXT("MotionController_Left_Thumbstick_Up", "MotionController (L) Thumbstick Up"), FKeyDetails::GamepadKey | FKeyDetails::NotBlueprintBindableKey | FKeyDetails::NotActionBindableKey));
+	AddKey(FKeyDetails(EKeys::MotionController_Left_Thumbstick_Down, LOCTEXT("MotionController_Left_Thumbstick_Down", "MotionController (L) Thumbstick Down"), FKeyDetails::GamepadKey | FKeyDetails::NotBlueprintBindableKey | FKeyDetails::NotActionBindableKey));
+	AddKey(FKeyDetails(EKeys::MotionController_Left_Thumbstick_Left, LOCTEXT("MotionController_Left_Thumbstick_Left", "MotionController (L) Thumbstick Left"), FKeyDetails::GamepadKey | FKeyDetails::NotBlueprintBindableKey | FKeyDetails::NotActionBindableKey));
+	AddKey(FKeyDetails(EKeys::MotionController_Left_Thumbstick_Right, LOCTEXT("MotionController_Left_Thumbstick_Right", "MotionController (L) Thumbstick Right"), FKeyDetails::GamepadKey | FKeyDetails::NotBlueprintBindableKey | FKeyDetails::NotActionBindableKey));
+
 	//	Right Controller
-	AddKey(FKeyDetails(EKeys::MotionController_Right_FaceButton1, LOCTEXT("MotionController_Right_FaceButton1", "MotionController (R) FaceButton1"), FKeyDetails::GamepadKey));
-	AddKey(FKeyDetails(EKeys::MotionController_Right_FaceButton2, LOCTEXT("MotionController_Right_FaceButton2", "MotionController (R) FaceButton2"), FKeyDetails::GamepadKey));
-	AddKey(FKeyDetails(EKeys::MotionController_Right_FaceButton3, LOCTEXT("MotionController_Right_FaceButton3", "MotionController (R) FaceButton3"), FKeyDetails::GamepadKey));
-	AddKey(FKeyDetails(EKeys::MotionController_Right_FaceButton4, LOCTEXT("MotionController_Right_FaceButton4", "MotionController (R) FaceButton4"), FKeyDetails::GamepadKey));
-	AddKey(FKeyDetails(EKeys::MotionController_Right_FaceButton5, LOCTEXT("MotionController_Right_FaceButton5", "MotionController (R) FaceButton5"), FKeyDetails::GamepadKey));
-	AddKey(FKeyDetails(EKeys::MotionController_Right_FaceButton6, LOCTEXT("MotionController_Right_FaceButton6", "MotionController (R) FaceButton6"), FKeyDetails::GamepadKey));
-	AddKey(FKeyDetails(EKeys::MotionController_Right_FaceButton7, LOCTEXT("MotionController_Right_FaceButton7", "MotionController (R) FaceButton7"), FKeyDetails::GamepadKey));
-	AddKey(FKeyDetails(EKeys::MotionController_Right_FaceButton8, LOCTEXT("MotionController_Right_FaceButton8", "MotionController (R) FaceButton8"), FKeyDetails::GamepadKey));
-	
-	AddKey(FKeyDetails(EKeys::MotionController_Right_Shoulder, LOCTEXT("MotionController_Right_Shoulder", "MotionController (R) Shoulder"), FKeyDetails::GamepadKey));
-	AddKey(FKeyDetails(EKeys::MotionController_Right_Trigger, LOCTEXT("MotionController_Right_Trigger", "MotionController (R) Trigger"), FKeyDetails::GamepadKey));
-	
-	AddKey(FKeyDetails(EKeys::MotionController_Right_Grip1, LOCTEXT("MotionController_Right_Grip1", "MotionController (R) Grip1"), FKeyDetails::GamepadKey));
-	AddKey(FKeyDetails(EKeys::MotionController_Right_Grip2, LOCTEXT("MotionController_Right_Grip2", "MotionController (R) Grip2"), FKeyDetails::GamepadKey));
-	
-	AddKey(FKeyDetails(EKeys::MotionController_Right_Thumbstick, LOCTEXT("MotionController_Right_Thumbstick", "MotionController (R) Thumbstick"), FKeyDetails::GamepadKey));
-	AddKey(FKeyDetails(EKeys::MotionController_Right_Thumbstick_Up, LOCTEXT("MotionController_Right_Thumbstick_Up", "MotionController (R) Thumbstick Up"), FKeyDetails::GamepadKey));
-	AddKey(FKeyDetails(EKeys::MotionController_Right_Thumbstick_Down, LOCTEXT("MotionController_Right_Thumbstick_Down", "MotionController (R) Thumbstick Down"), FKeyDetails::GamepadKey));
-	AddKey(FKeyDetails(EKeys::MotionController_Right_Thumbstick_Left, LOCTEXT("MotionController_Right_Thumbstick_Left", "MotionController (R) Thumbstick Left"), FKeyDetails::GamepadKey));
-	AddKey(FKeyDetails(EKeys::MotionController_Right_Thumbstick_Right, LOCTEXT("MotionController_Right_Thumbstick_Right", "MotionController (R) Thumbstick Right"), FKeyDetails::GamepadKey));
-	
+	AddKey(FKeyDetails(EKeys::MotionController_Right_FaceButton1, LOCTEXT("MotionController_Right_FaceButton1", "MotionController (R) FaceButton1"), FKeyDetails::GamepadKey | FKeyDetails::NotBlueprintBindableKey | FKeyDetails::NotActionBindableKey));
+	AddKey(FKeyDetails(EKeys::MotionController_Right_FaceButton2, LOCTEXT("MotionController_Right_FaceButton2", "MotionController (R) FaceButton2"), FKeyDetails::GamepadKey | FKeyDetails::NotBlueprintBindableKey | FKeyDetails::NotActionBindableKey));
+	AddKey(FKeyDetails(EKeys::MotionController_Right_FaceButton3, LOCTEXT("MotionController_Right_FaceButton3", "MotionController (R) FaceButton3"), FKeyDetails::GamepadKey | FKeyDetails::NotBlueprintBindableKey | FKeyDetails::NotActionBindableKey));
+	AddKey(FKeyDetails(EKeys::MotionController_Right_FaceButton4, LOCTEXT("MotionController_Right_FaceButton4", "MotionController (R) FaceButton4"), FKeyDetails::GamepadKey | FKeyDetails::NotBlueprintBindableKey | FKeyDetails::NotActionBindableKey));
+	AddKey(FKeyDetails(EKeys::MotionController_Right_FaceButton5, LOCTEXT("MotionController_Right_FaceButton5", "MotionController (R) FaceButton5"), FKeyDetails::GamepadKey | FKeyDetails::NotBlueprintBindableKey | FKeyDetails::NotActionBindableKey));
+	AddKey(FKeyDetails(EKeys::MotionController_Right_FaceButton6, LOCTEXT("MotionController_Right_FaceButton6", "MotionController (R) FaceButton6"), FKeyDetails::GamepadKey | FKeyDetails::NotBlueprintBindableKey | FKeyDetails::NotActionBindableKey));
+	AddKey(FKeyDetails(EKeys::MotionController_Right_FaceButton7, LOCTEXT("MotionController_Right_FaceButton7", "MotionController (R) FaceButton7"), FKeyDetails::GamepadKey | FKeyDetails::NotBlueprintBindableKey | FKeyDetails::NotActionBindableKey));
+	AddKey(FKeyDetails(EKeys::MotionController_Right_FaceButton8, LOCTEXT("MotionController_Right_FaceButton8", "MotionController (R) FaceButton8"), FKeyDetails::GamepadKey | FKeyDetails::NotBlueprintBindableKey | FKeyDetails::NotActionBindableKey));
+
+	AddKey(FKeyDetails(EKeys::MotionController_Right_Shoulder, LOCTEXT("MotionController_Right_Shoulder", "MotionController (R) Shoulder"), FKeyDetails::GamepadKey | FKeyDetails::NotBlueprintBindableKey | FKeyDetails::NotActionBindableKey));
+	AddKey(FKeyDetails(EKeys::MotionController_Right_Trigger, LOCTEXT("MotionController_Right_Trigger", "MotionController (R) Trigger"), FKeyDetails::GamepadKey | FKeyDetails::NotBlueprintBindableKey | FKeyDetails::NotActionBindableKey));
+
+	AddKey(FKeyDetails(EKeys::MotionController_Right_Grip1, LOCTEXT("MotionController_Right_Grip1", "MotionController (R) Grip1"), FKeyDetails::GamepadKey | FKeyDetails::NotBlueprintBindableKey | FKeyDetails::NotActionBindableKey));
+	AddKey(FKeyDetails(EKeys::MotionController_Right_Grip2, LOCTEXT("MotionController_Right_Grip2", "MotionController (R) Grip2"), FKeyDetails::GamepadKey | FKeyDetails::NotBlueprintBindableKey | FKeyDetails::NotActionBindableKey));
+
+	AddKey(FKeyDetails(EKeys::MotionController_Right_Thumbstick, LOCTEXT("MotionController_Right_Thumbstick", "MotionController (R) Thumbstick"), FKeyDetails::GamepadKey | FKeyDetails::NotBlueprintBindableKey | FKeyDetails::NotActionBindableKey));
+	AddKey(FKeyDetails(EKeys::MotionController_Right_Thumbstick_Up, LOCTEXT("MotionController_Right_Thumbstick_Up", "MotionController (R) Thumbstick Up"), FKeyDetails::GamepadKey | FKeyDetails::NotBlueprintBindableKey | FKeyDetails::NotActionBindableKey));
+	AddKey(FKeyDetails(EKeys::MotionController_Right_Thumbstick_Down, LOCTEXT("MotionController_Right_Thumbstick_Down", "MotionController (R) Thumbstick Down"), FKeyDetails::GamepadKey | FKeyDetails::NotBlueprintBindableKey | FKeyDetails::NotActionBindableKey));
+	AddKey(FKeyDetails(EKeys::MotionController_Right_Thumbstick_Left, LOCTEXT("MotionController_Right_Thumbstick_Left", "MotionController (R) Thumbstick Left"), FKeyDetails::GamepadKey | FKeyDetails::NotBlueprintBindableKey | FKeyDetails::NotActionBindableKey));
+	AddKey(FKeyDetails(EKeys::MotionController_Right_Thumbstick_Right, LOCTEXT("MotionController_Right_Thumbstick_Right", "MotionController (R) Thumbstick Right"), FKeyDetails::GamepadKey | FKeyDetails::NotBlueprintBindableKey | FKeyDetails::NotActionBindableKey));
+
 	//   Motion Controller Axes
 	//	Left Controller
-	AddKey(FKeyDetails(EKeys::MotionController_Left_Thumbstick_X, LOCTEXT("MotionController_Left_Thumbstick_X", "MotionController (L) Thumbstick X"), FKeyDetails::GamepadKey | FKeyDetails::FloatAxis));
-	AddKey(FKeyDetails(EKeys::MotionController_Left_Thumbstick_Y, LOCTEXT("MotionController_Left_Thumbstick_Y", "MotionController (L) Thumbstick Y"), FKeyDetails::GamepadKey | FKeyDetails::FloatAxis));
-	AddKey(FKeyDetails(EKeys::MotionController_Left_TriggerAxis, LOCTEXT("MotionController_Left_TriggerAxis", "MotionController (L) TriggerAxis"), FKeyDetails::GamepadKey | FKeyDetails::FloatAxis));
-	AddKey(FKeyDetails(EKeys::MotionController_Left_Grip1Axis, LOCTEXT("MotionController_Left_Grip1Axis", "MotionController (L) Grip1 Axis"), FKeyDetails::GamepadKey | FKeyDetails::FloatAxis));
-	AddKey(FKeyDetails(EKeys::MotionController_Left_Grip2Axis, LOCTEXT("MotionController_Left_Grip2Axis", "MotionController (L) Grip2 Axis"), FKeyDetails::GamepadKey | FKeyDetails::FloatAxis));
+	AddKey(FKeyDetails(EKeys::MotionController_Left_Thumbstick_X, LOCTEXT("MotionController_Left_Thumbstick_X", "MotionController (L) Thumbstick X"), FKeyDetails::GamepadKey | FKeyDetails::FloatAxis | FKeyDetails::NotBlueprintBindableKey | FKeyDetails::NotActionBindableKey));
+	AddKey(FKeyDetails(EKeys::MotionController_Left_Thumbstick_Y, LOCTEXT("MotionController_Left_Thumbstick_Y", "MotionController (L) Thumbstick Y"), FKeyDetails::GamepadKey | FKeyDetails::FloatAxis | FKeyDetails::NotBlueprintBindableKey | FKeyDetails::NotActionBindableKey));
+	AddKey(FKeyDetails(EKeys::MotionController_Left_TriggerAxis, LOCTEXT("MotionController_Left_TriggerAxis", "MotionController (L) TriggerAxis"), FKeyDetails::GamepadKey | FKeyDetails::FloatAxis | FKeyDetails::NotBlueprintBindableKey | FKeyDetails::NotActionBindableKey));
+	AddKey(FKeyDetails(EKeys::MotionController_Left_Grip1Axis, LOCTEXT("MotionController_Left_Grip1Axis", "MotionController (L) Grip1 Axis"), FKeyDetails::GamepadKey | FKeyDetails::FloatAxis | FKeyDetails::NotBlueprintBindableKey | FKeyDetails::NotActionBindableKey));
+	AddKey(FKeyDetails(EKeys::MotionController_Left_Grip2Axis, LOCTEXT("MotionController_Left_Grip2Axis", "MotionController (L) Grip2 Axis"), FKeyDetails::GamepadKey | FKeyDetails::FloatAxis | FKeyDetails::NotBlueprintBindableKey | FKeyDetails::NotActionBindableKey));
 
 	//	Right Controller
-	AddKey(FKeyDetails(EKeys::MotionController_Right_Thumbstick_X, LOCTEXT("MotionController_Right_Thumbstick_X", "MotionController (R) Thumbstick X"), FKeyDetails::GamepadKey | FKeyDetails::FloatAxis));
-	AddKey(FKeyDetails(EKeys::MotionController_Right_Thumbstick_Y, LOCTEXT("MotionController_Right_Thumbstick_Y", "MotionController (R) Thumbstick Y"), FKeyDetails::GamepadKey | FKeyDetails::FloatAxis));
-	AddKey(FKeyDetails(EKeys::MotionController_Right_TriggerAxis, LOCTEXT("MotionController_Right_TriggerAxis", "MotionController (R) TriggerAxis"), FKeyDetails::GamepadKey | FKeyDetails::FloatAxis));
-	AddKey(FKeyDetails(EKeys::MotionController_Right_Grip1Axis, LOCTEXT("MotionController_Right_Grip1Axis", "MotionController (R) Grip1 Axis"), FKeyDetails::GamepadKey | FKeyDetails::FloatAxis));
-	AddKey(FKeyDetails(EKeys::MotionController_Right_Grip2Axis, LOCTEXT("MotionController_Right_Grip2Axis", "MotionController (R) Grip2 Axis"), FKeyDetails::GamepadKey | FKeyDetails::FloatAxis));
+	AddKey(FKeyDetails(EKeys::MotionController_Right_Thumbstick_X, LOCTEXT("MotionController_Right_Thumbstick_X", "MotionController (R) Thumbstick X"), FKeyDetails::GamepadKey | FKeyDetails::FloatAxis | FKeyDetails::NotBlueprintBindableKey | FKeyDetails::NotActionBindableKey));
+	AddKey(FKeyDetails(EKeys::MotionController_Right_Thumbstick_Y, LOCTEXT("MotionController_Right_Thumbstick_Y", "MotionController (R) Thumbstick Y"), FKeyDetails::GamepadKey | FKeyDetails::FloatAxis | FKeyDetails::NotBlueprintBindableKey | FKeyDetails::NotActionBindableKey));
+	AddKey(FKeyDetails(EKeys::MotionController_Right_TriggerAxis, LOCTEXT("MotionController_Right_TriggerAxis", "MotionController (R) TriggerAxis"), FKeyDetails::GamepadKey | FKeyDetails::FloatAxis | FKeyDetails::NotBlueprintBindableKey | FKeyDetails::NotActionBindableKey));
+	AddKey(FKeyDetails(EKeys::MotionController_Right_Grip1Axis, LOCTEXT("MotionController_Right_Grip1Axis", "MotionController (R) Grip1 Axis"), FKeyDetails::GamepadKey | FKeyDetails::FloatAxis | FKeyDetails::NotBlueprintBindableKey | FKeyDetails::NotActionBindableKey));
+	AddKey(FKeyDetails(EKeys::MotionController_Right_Grip2Axis, LOCTEXT("MotionController_Right_Grip2Axis", "MotionController (R) Grip2 Axis"), FKeyDetails::GamepadKey | FKeyDetails::FloatAxis | FKeyDetails::NotBlueprintBindableKey | FKeyDetails::NotActionBindableKey));
 
 	// PS4-specific
 	AddMenuCategoryDisplayInfo("PS4", LOCTEXT("PS4SubCategory", "PS4"), TEXT("GraphEditor.PadEvent_16x"));
-	AddKey(FKeyDetails(EKeys::PS4_Special, LOCTEXT("PS4_Special", "PS4_Special"), FKeyDetails::NotBlueprintBindableKey, "PS4"));
+	AddKey(FKeyDetails(EKeys::PS4_Special, LOCTEXT("PS4_Special", "PS4_Special"), FKeyDetails::NotBlueprintBindableKey | FKeyDetails::NotActionBindableKey, "PS4"));
 	AddKey(FKeyDetails(EKeys::Gamepad_Special_Left_X, LOCTEXT("PS4_Gamepad_Special_Left_X", "PS4 Touchpad Button X Axis"), FKeyDetails::GamepadKey | FKeyDetails::FloatAxis, "PS4"));
 	AddKey(FKeyDetails(EKeys::Gamepad_Special_Left_Y, LOCTEXT("PS4_Gamepad_Special_Left_Y", "PS4 Touchpad Button Y Axis"), FKeyDetails::GamepadKey | FKeyDetails::FloatAxis, "PS4"));
 
@@ -709,6 +856,164 @@ void EKeys::Initialize()
 	AddKey(FKeyDetails(EKeys::Android_Volume_Down, LOCTEXT("Android_Volume_Down", "Android Volume Down"), FKeyDetails::GamepadKey, "Android"));
 	AddKey(FKeyDetails(EKeys::Android_Menu, LOCTEXT("Android_Menu", "Android Menu"), FKeyDetails::GamepadKey, "Android"));
 
+	// Google Daydream
+	AddMenuCategoryDisplayInfo("Daydream", LOCTEXT("DaydreamSubCategory", "Google Daydream"), TEXT("GraphEditor.PadEvent_16x"));
+
+	AddKey(FKeyDetails(EKeys::Daydream_Left_Select_Click, LOCTEXT("Daydream_Left_Select_Click", "Daydream (L) Select"), FKeyDetails::GamepadKey | FKeyDetails::NotBlueprintBindableKey, "Daydream"));
+	AddKey(FKeyDetails(EKeys::Daydream_Left_Trackpad_X, LOCTEXT("Daydream_Left_Trackpad_X", "Daydream (L) Trackpad X"), FKeyDetails::GamepadKey | FKeyDetails::FloatAxis | FKeyDetails::NotBlueprintBindableKey, "Daydream"));
+	AddKey(FKeyDetails(EKeys::Daydream_Left_Trackpad_Y, LOCTEXT("Daydream_Left_Trackpad_Y", "Daydream (L) Trackpad Y"), FKeyDetails::GamepadKey | FKeyDetails::FloatAxis | FKeyDetails::NotBlueprintBindableKey, "Daydream"));
+	AddKey(FKeyDetails(EKeys::Daydream_Left_Trackpad_Click, LOCTEXT("Daydream_Left_Trackpad_Click", "Daydream (L) Trackpad"), FKeyDetails::GamepadKey | FKeyDetails::NotBlueprintBindableKey, "Daydream"));
+	AddKey(FKeyDetails(EKeys::Daydream_Left_Trackpad_Touch, LOCTEXT("Daydream_Left_Trackpad_Touch", "Daydream (L) Trackpad Touch"), FKeyDetails::GamepadKey | FKeyDetails::NotBlueprintBindableKey, "Daydream"));
+	AddKey(FKeyDetails(EKeys::Daydream_Right_Select_Click, LOCTEXT("Daydream_Right_Select_Click", "Daydream (R) Select"), FKeyDetails::GamepadKey | FKeyDetails::NotBlueprintBindableKey, "Daydream"));
+	AddKey(FKeyDetails(EKeys::Daydream_Right_Trackpad_X, LOCTEXT("Daydream_Right_Trackpad_X", "Daydream (R) Trackpad X"), FKeyDetails::GamepadKey | FKeyDetails::FloatAxis | FKeyDetails::NotBlueprintBindableKey, "Daydream"));
+	AddKey(FKeyDetails(EKeys::Daydream_Right_Trackpad_Y, LOCTEXT("Daydream_Right_Trackpad_Y", "Daydream (R) Trackpad Y"), FKeyDetails::GamepadKey | FKeyDetails::FloatAxis | FKeyDetails::NotBlueprintBindableKey, "Daydream"));
+	AddKey(FKeyDetails(EKeys::Daydream_Right_Trackpad_Click, LOCTEXT("Daydream_Right_Trackpad_Click", "Daydream (R) Trackpad"), FKeyDetails::GamepadKey | FKeyDetails::NotBlueprintBindableKey, "Daydream"));
+	AddKey(FKeyDetails(EKeys::Daydream_Right_Trackpad_Touch, LOCTEXT("Daydream_Right_Trackpad_Touch", "Daydream (R) Trackpad Touch"), FKeyDetails::GamepadKey | FKeyDetails::NotBlueprintBindableKey, "Daydream"));
+
+	// HTC Vive Controller
+	AddMenuCategoryDisplayInfo("Vive", LOCTEXT("ViveSubCategory", "HTC Vive"), TEXT("GraphEditor.PadEvent_16x"));
+
+	AddKey(FKeyDetails(EKeys::Vive_Left_System_Click, LOCTEXT("Vive_Left_System_Click", "Vive (L) System"), FKeyDetails::GamepadKey | FKeyDetails::NotBlueprintBindableKey, "Vive"));
+	AddKey(FKeyDetails(EKeys::Vive_Left_Grip_Click, LOCTEXT("Vive_Left_Grip_Click", "Vive (L) Grip"), FKeyDetails::GamepadKey | FKeyDetails::NotBlueprintBindableKey, "Vive"));
+	AddKey(FKeyDetails(EKeys::Vive_Left_Menu_Click, LOCTEXT("Vive_Left_Menu_Click", "Vive (L) Menu"), FKeyDetails::GamepadKey | FKeyDetails::NotBlueprintBindableKey, "Vive"));
+	AddKey(FKeyDetails(EKeys::Vive_Left_Trigger_Click, LOCTEXT("Vive_Left_Trigger_Click", "Vive (L) Trigger"), FKeyDetails::GamepadKey | FKeyDetails::NotBlueprintBindableKey, "Vive"));
+	AddKey(FKeyDetails(EKeys::Vive_Left_Trigger_Axis, LOCTEXT("Vive_Left_Trigger_Axis", "Vive (L) Trigger Axis"), FKeyDetails::GamepadKey | FKeyDetails::FloatAxis | FKeyDetails::NotBlueprintBindableKey, "Vive"));
+	AddKey(FKeyDetails(EKeys::Vive_Left_Trackpad_X, LOCTEXT("Vive_Left_Trackpad_X", "Vive (L) Trackpad X"), FKeyDetails::GamepadKey | FKeyDetails::FloatAxis | FKeyDetails::NotBlueprintBindableKey, "Vive"));
+	AddKey(FKeyDetails(EKeys::Vive_Left_Trackpad_Y, LOCTEXT("Vive_Left_Trackpad_Y", "Vive (L) Trackpad Y"), FKeyDetails::GamepadKey | FKeyDetails::FloatAxis | FKeyDetails::NotBlueprintBindableKey, "Vive"));
+	AddKey(FKeyDetails(EKeys::Vive_Left_Trackpad_Click, LOCTEXT("Vive_Left_Trackpad_Click", "Vive (L) Trackpad"), FKeyDetails::GamepadKey | FKeyDetails::NotBlueprintBindableKey, "Vive"));
+	AddKey(FKeyDetails(EKeys::Vive_Left_Trackpad_Touch, LOCTEXT("Vive_Left_Trackpad_Touch", "Vive (L) Trackpad Touch"), FKeyDetails::GamepadKey | FKeyDetails::NotBlueprintBindableKey, "Vive"));
+	AddKey(FKeyDetails(EKeys::Vive_Right_System_Click, LOCTEXT("Vive_Right_System_Click", "Vive (R) System"), FKeyDetails::GamepadKey | FKeyDetails::NotBlueprintBindableKey, "Vive"));
+	AddKey(FKeyDetails(EKeys::Vive_Right_Grip_Click, LOCTEXT("Vive_Right_Grip_Click", "Vive (R) Grip"), FKeyDetails::GamepadKey | FKeyDetails::NotBlueprintBindableKey, "Vive"));
+	AddKey(FKeyDetails(EKeys::Vive_Right_Menu_Click, LOCTEXT("Vive_Right_Menu_Click", "Vive (R) Menu"), FKeyDetails::GamepadKey | FKeyDetails::NotBlueprintBindableKey, "Vive"));
+	AddKey(FKeyDetails(EKeys::Vive_Right_Trigger_Click, LOCTEXT("Vive_Right_Trigger_Click", "Vive (R) Trigger"), FKeyDetails::GamepadKey | FKeyDetails::NotBlueprintBindableKey, "Vive"));
+	AddKey(FKeyDetails(EKeys::Vive_Right_Trigger_Axis, LOCTEXT("Vive_Right_Trigger_Axis", "Vive (R) Trigger Axis"), FKeyDetails::GamepadKey | FKeyDetails::FloatAxis | FKeyDetails::NotBlueprintBindableKey, "Vive"));
+	AddKey(FKeyDetails(EKeys::Vive_Right_Trackpad_X, LOCTEXT("Vive_Right_Trackpad_X", "Vive (R) Trackpad X"), FKeyDetails::GamepadKey | FKeyDetails::FloatAxis | FKeyDetails::NotBlueprintBindableKey, "Vive"));
+	AddKey(FKeyDetails(EKeys::Vive_Right_Trackpad_Y, LOCTEXT("Vive_Right_Trackpad_Y", "Vive (R) Trackpad Y"), FKeyDetails::GamepadKey | FKeyDetails::FloatAxis | FKeyDetails::NotBlueprintBindableKey, "Vive"));
+	AddKey(FKeyDetails(EKeys::Vive_Right_Trackpad_Click, LOCTEXT("Vive_Right_Trackpad_Click", "Vive (R) Trackpad"), FKeyDetails::GamepadKey | FKeyDetails::NotBlueprintBindableKey, "Vive"));
+	AddKey(FKeyDetails(EKeys::Vive_Right_Trackpad_Touch, LOCTEXT("Vive_Right_Trackpad_Touch", "Vive (R) Trackpad Touch"), FKeyDetails::GamepadKey | FKeyDetails::NotBlueprintBindableKey, "Vive"));
+
+	// Microsoft Mixed Reality Motion Controller
+	AddMenuCategoryDisplayInfo("MixedReality", LOCTEXT("MixedRealitySubCategory", "Windows Mixed Reality"), TEXT("GraphEditor.PadEvent_16x"));
+
+	AddKey(FKeyDetails(EKeys::MixedReality_Left_Menu_Click, LOCTEXT("MixedReality_Left_Menu_Click", "Mixed Reality (L) Menu"), FKeyDetails::GamepadKey | FKeyDetails::NotBlueprintBindableKey, "MixedReality"));
+	AddKey(FKeyDetails(EKeys::MixedReality_Left_Grip_Click, LOCTEXT("MixedReality_Left_Grip_Click", "Mixed Reality (L) Grip"), FKeyDetails::GamepadKey | FKeyDetails::NotBlueprintBindableKey, "MixedReality"));
+	AddKey(FKeyDetails(EKeys::MixedReality_Left_Trigger_Click, LOCTEXT("MixedReality_Left_Trigger_Click", "Mixed Reality (L) Trigger"), FKeyDetails::GamepadKey | FKeyDetails::NotBlueprintBindableKey, "MixedReality"));
+	AddKey(FKeyDetails(EKeys::MixedReality_Left_Trigger_Axis, LOCTEXT("MixedReality_Left_Trigger_Axis", "Mixed Reality (L) Trigger Axis"), FKeyDetails::GamepadKey | FKeyDetails::FloatAxis | FKeyDetails::NotBlueprintBindableKey, "MixedReality"));
+	AddKey(FKeyDetails(EKeys::MixedReality_Left_Thumbstick_X, LOCTEXT("MixedReality_Left_Thumbstick_X", "Mixed Reality (L) Thumbstick X"), FKeyDetails::GamepadKey | FKeyDetails::FloatAxis | FKeyDetails::NotBlueprintBindableKey, "MixedReality"));
+	AddKey(FKeyDetails(EKeys::MixedReality_Left_Thumbstick_Y, LOCTEXT("MixedReality_Left_Thumbstick_Y", "Mixed Reality (L) Thumbstick Y"), FKeyDetails::GamepadKey | FKeyDetails::FloatAxis | FKeyDetails::NotBlueprintBindableKey, "MixedReality"));
+	AddKey(FKeyDetails(EKeys::MixedReality_Left_Thumbstick_Click, LOCTEXT("MixedReality_Left_Thumbstick_Click", "Mixed Reality (L) Thumbstick"), FKeyDetails::GamepadKey | FKeyDetails::NotBlueprintBindableKey, "MixedReality"));
+	AddKey(FKeyDetails(EKeys::MixedReality_Left_Trackpad_X, LOCTEXT("MixedReality_Left_Trackpad_X", "Mixed Reality (L) Trackpad X"), FKeyDetails::GamepadKey | FKeyDetails::FloatAxis | FKeyDetails::NotBlueprintBindableKey, "MixedReality"));
+	AddKey(FKeyDetails(EKeys::MixedReality_Left_Trackpad_Y, LOCTEXT("MixedReality_Left_Trackpad_Y", "Mixed Reality (L) Trackpad Y"), FKeyDetails::GamepadKey | FKeyDetails::FloatAxis | FKeyDetails::NotBlueprintBindableKey, "MixedReality"));
+	AddKey(FKeyDetails(EKeys::MixedReality_Left_Trackpad_Click, LOCTEXT("MixedReality_Left_Trackpad_Click", "Mixed Reality (L) Trackpad"), FKeyDetails::GamepadKey | FKeyDetails::NotBlueprintBindableKey, "MixedReality"));
+	AddKey(FKeyDetails(EKeys::MixedReality_Left_Trackpad_Touch, LOCTEXT("MixedReality_Left_Trackpad_Touch", "Mixed Reality (L) Trackpad Touch"), FKeyDetails::GamepadKey | FKeyDetails::NotBlueprintBindableKey, "MixedReality"));
+	AddKey(FKeyDetails(EKeys::MixedReality_Right_Menu_Click, LOCTEXT("MixedReality_Right_Menu_Click", "Mixed Reality (R) Menu"), FKeyDetails::GamepadKey | FKeyDetails::NotBlueprintBindableKey, "MixedReality"));
+	AddKey(FKeyDetails(EKeys::MixedReality_Right_Grip_Click, LOCTEXT("MixedReality_Right_Grip_Click", "Mixed Reality (R) Grip"), FKeyDetails::GamepadKey | FKeyDetails::NotBlueprintBindableKey, "MixedReality"));
+	AddKey(FKeyDetails(EKeys::MixedReality_Right_Trigger_Click, LOCTEXT("MixedReality_Right_Trigger_Click", "Mixed Reality (R) Trigger"), FKeyDetails::GamepadKey | FKeyDetails::NotBlueprintBindableKey, "MixedReality"));
+	AddKey(FKeyDetails(EKeys::MixedReality_Right_Trigger_Axis, LOCTEXT("MixedReality_Right_Trigger_Axis", "Mixed Reality (R) Trigger Axis"), FKeyDetails::GamepadKey | FKeyDetails::FloatAxis | FKeyDetails::NotBlueprintBindableKey, "MixedReality"));
+	AddKey(FKeyDetails(EKeys::MixedReality_Right_Thumbstick_X, LOCTEXT("MixedReality_Right_Thumbstick_X", "Mixed Reality (R) Thumbstick X"), FKeyDetails::GamepadKey | FKeyDetails::FloatAxis | FKeyDetails::NotBlueprintBindableKey, "MixedReality"));
+	AddKey(FKeyDetails(EKeys::MixedReality_Right_Thumbstick_Y, LOCTEXT("MixedReality_Right_Thumbstick_Y", "Mixed Reality (R) Thumbstick Y"), FKeyDetails::GamepadKey | FKeyDetails::FloatAxis | FKeyDetails::NotBlueprintBindableKey, "MixedReality"));
+	AddKey(FKeyDetails(EKeys::MixedReality_Right_Thumbstick_Click, LOCTEXT("MixedReality_Right_Thumbstick_Click", "Mixed Reality (R) Thumbstick"), FKeyDetails::GamepadKey | FKeyDetails::NotBlueprintBindableKey, "MixedReality"));
+	AddKey(FKeyDetails(EKeys::MixedReality_Right_Trackpad_X, LOCTEXT("MixedReality_Right_Trackpad_X", "Mixed Reality (R) Trackpad X"), FKeyDetails::GamepadKey | FKeyDetails::FloatAxis | FKeyDetails::NotBlueprintBindableKey, "MixedReality"));
+	AddKey(FKeyDetails(EKeys::MixedReality_Right_Trackpad_Y, LOCTEXT("MixedReality_Right_Trackpad_Y", "Mixed Reality (R) Trackpad Y"), FKeyDetails::GamepadKey | FKeyDetails::FloatAxis | FKeyDetails::NotBlueprintBindableKey, "MixedReality"));
+	AddKey(FKeyDetails(EKeys::MixedReality_Right_Trackpad_Click, LOCTEXT("MixedReality_Right_Trackpad_Click", "Mixed Reality (R) Trackpad"), FKeyDetails::GamepadKey | FKeyDetails::NotBlueprintBindableKey, "MixedReality"));
+	AddKey(FKeyDetails(EKeys::MixedReality_Right_Trackpad_Touch, LOCTEXT("MixedReality_Right_Trackpad_Touch", "Mixed Reality (R) Trackpad Touch"), FKeyDetails::GamepadKey | FKeyDetails::NotBlueprintBindableKey, "MixedReality"));
+
+	// Oculus Go Controller
+	AddMenuCategoryDisplayInfo("OculusGo", LOCTEXT("OculusGoSubCategory", "Oculus Go"), TEXT("GraphEditor.PadEvent_16x"));
+
+	AddKey(FKeyDetails(EKeys::OculusGo_Left_System_Click, LOCTEXT("OculusGo_Left_System_Click", "Oculus Go (L) System"), FKeyDetails::GamepadKey | FKeyDetails::NotBlueprintBindableKey, "OculusGo"));
+	AddKey(FKeyDetails(EKeys::OculusGo_Left_Back_Click, LOCTEXT("OculusGo_Left_Back_Click", "Oculus Go (L) Back"), FKeyDetails::GamepadKey | FKeyDetails::NotBlueprintBindableKey, "OculusGo"));
+	AddKey(FKeyDetails(EKeys::OculusGo_Left_Trigger_Click, LOCTEXT("OculusGo_Left_Trigger_Click", "Oculus Go (L) Trigger"), FKeyDetails::GamepadKey | FKeyDetails::NotBlueprintBindableKey, "OculusGo"));
+	AddKey(FKeyDetails(EKeys::OculusGo_Left_Trackpad_X, LOCTEXT("OculusGo_Left_Trackpad_X", "Oculus Go (L) Trackpad X"), FKeyDetails::GamepadKey | FKeyDetails::FloatAxis | FKeyDetails::NotBlueprintBindableKey, "OculusGo"));
+	AddKey(FKeyDetails(EKeys::OculusGo_Left_Trackpad_Y, LOCTEXT("OculusGo_Left_Trackpad_Y", "Oculus Go (L) Trackpad Y"), FKeyDetails::GamepadKey | FKeyDetails::FloatAxis | FKeyDetails::NotBlueprintBindableKey, "OculusGo"));
+	AddKey(FKeyDetails(EKeys::OculusGo_Left_Trackpad_Click, LOCTEXT("OculusGo_Left_Trackpad_Click", "Oculus Go (L) Trackpad"), FKeyDetails::GamepadKey | FKeyDetails::NotBlueprintBindableKey, "OculusGo"));
+	AddKey(FKeyDetails(EKeys::OculusGo_Left_Trackpad_Touch, LOCTEXT("OculusGo_Left_Trackpad_Touch", "Oculus Go (L) Trackpad Touch"), FKeyDetails::GamepadKey | FKeyDetails::NotBlueprintBindableKey, "OculusGo"));
+	AddKey(FKeyDetails(EKeys::OculusGo_Right_System_Click, LOCTEXT("OculusGo_Right_System_Click", "Oculus Go (R) System"), FKeyDetails::GamepadKey | FKeyDetails::NotBlueprintBindableKey, "OculusGo"));
+	AddKey(FKeyDetails(EKeys::OculusGo_Right_Back_Click, LOCTEXT("OculusGo_Right_Back_Click", "Oculus Go (R) Back"), FKeyDetails::GamepadKey | FKeyDetails::NotBlueprintBindableKey, "OculusGo"));
+	AddKey(FKeyDetails(EKeys::OculusGo_Right_Trigger_Click, LOCTEXT("OculusGo_Right_Trigger_Click", "Oculus Go (R) Trigger"), FKeyDetails::GamepadKey | FKeyDetails::NotBlueprintBindableKey, "OculusGo"));
+	AddKey(FKeyDetails(EKeys::OculusGo_Right_Trackpad_X, LOCTEXT("OculusGo_Right_Trackpad_X", "Oculus Go (R) Trackpad X"), FKeyDetails::GamepadKey | FKeyDetails::FloatAxis | FKeyDetails::NotBlueprintBindableKey, "OculusGo"));
+	AddKey(FKeyDetails(EKeys::OculusGo_Right_Trackpad_Y, LOCTEXT("OculusGo_Right_Trackpad_Y", "Oculus Go (R) Trackpad Y"), FKeyDetails::GamepadKey | FKeyDetails::FloatAxis | FKeyDetails::NotBlueprintBindableKey, "OculusGo"));
+	AddKey(FKeyDetails(EKeys::OculusGo_Right_Trackpad_Click, LOCTEXT("OculusGo_Right_Trackpad_Click", "Oculus Go (R) Trackpad"), FKeyDetails::GamepadKey | FKeyDetails::NotBlueprintBindableKey, "OculusGo"));
+	AddKey(FKeyDetails(EKeys::OculusGo_Right_Trackpad_Touch, LOCTEXT("OculusGo_Right_Trackpad_Touch", "Oculus Go (R) Trackpad Touch"), FKeyDetails::GamepadKey | FKeyDetails::NotBlueprintBindableKey, "OculusGo"));
+
+	// Oculus Touch Controller
+	AddMenuCategoryDisplayInfo("OculusTouch", LOCTEXT("OculusTouchSubCategory", "Oculus Touch"), TEXT("GraphEditor.PadEvent_16x"));
+
+	AddKey(FKeyDetails(EKeys::OculusTouch_Left_X_Click, LOCTEXT("OculusTouch_Left_X_Click", "Oculus Touch (L) X Press"), FKeyDetails::GamepadKey | FKeyDetails::NotBlueprintBindableKey, "OculusTouch"));
+	AddKey(FKeyDetails(EKeys::OculusTouch_Left_Y_Click, LOCTEXT("OculusTouch_Left_Y_Click", "Oculus Touch (L) Y Press"), FKeyDetails::GamepadKey | FKeyDetails::NotBlueprintBindableKey, "OculusTouch"));
+	AddKey(FKeyDetails(EKeys::OculusTouch_Left_X_Touch, LOCTEXT("OculusTouch_Left_X_Touch", "Oculus Touch (L) X Touch"), FKeyDetails::GamepadKey | FKeyDetails::NotBlueprintBindableKey, "OculusTouch"));
+	AddKey(FKeyDetails(EKeys::OculusTouch_Left_Y_Touch, LOCTEXT("OculusTouch_Left_Y_Touch", "Oculus Touch (L) Y Touch"), FKeyDetails::GamepadKey | FKeyDetails::NotBlueprintBindableKey, "OculusTouch"));
+	AddKey(FKeyDetails(EKeys::OculusTouch_Left_Menu_Click, LOCTEXT("OculusTouch_Left_Menu_Click", "Oculus Touch (L) Menu"), FKeyDetails::GamepadKey | FKeyDetails::NotBlueprintBindableKey, "OculusTouch"));
+	AddKey(FKeyDetails(EKeys::OculusTouch_Left_Grip_Click, LOCTEXT("OculusTouch_Left_Grip_Click", "Oculus Touch (L) Grip"), FKeyDetails::GamepadKey | FKeyDetails::NotBlueprintBindableKey, "OculusTouch"));
+	AddKey(FKeyDetails(EKeys::OculusTouch_Left_Grip_Axis, LOCTEXT("OculusTouch_Left_Grip_Axis", "Oculus Touch (L) Grip Axis"), FKeyDetails::GamepadKey | FKeyDetails::FloatAxis | FKeyDetails::NotBlueprintBindableKey, "OculusTouch"));
+	AddKey(FKeyDetails(EKeys::OculusTouch_Left_Trigger_Click, LOCTEXT("OculusTouch_Left_Trigger_Click", "Oculus Touch (L) Trigger"), FKeyDetails::GamepadKey | FKeyDetails::NotBlueprintBindableKey, "OculusTouch"));
+	AddKey(FKeyDetails(EKeys::OculusTouch_Left_Trigger_Axis, LOCTEXT("OculusTouch_Left_Trigger_Axis", "Oculus Touch (L) Trigger Axis"), FKeyDetails::GamepadKey | FKeyDetails::FloatAxis | FKeyDetails::NotBlueprintBindableKey, "OculusTouch"));
+	AddKey(FKeyDetails(EKeys::OculusTouch_Left_Trigger_Touch, LOCTEXT("OculusTouch_Left_Trigger_Touch", "Oculus Touch (L) Trigger Touch"), FKeyDetails::GamepadKey | FKeyDetails::NotBlueprintBindableKey, "OculusTouch"));
+	AddKey(FKeyDetails(EKeys::OculusTouch_Left_Thumbstick_X, LOCTEXT("OculusTouch_Left_Thumbstick_X", "Oculus Touch (L) Thumbstick X"), FKeyDetails::GamepadKey | FKeyDetails::FloatAxis | FKeyDetails::NotBlueprintBindableKey, "OculusTouch"));
+	AddKey(FKeyDetails(EKeys::OculusTouch_Left_Thumbstick_Y, LOCTEXT("OculusTouch_Left_Thumbstick_Y", "Oculus Touch (L) Thumbstick Y"), FKeyDetails::GamepadKey | FKeyDetails::FloatAxis | FKeyDetails::NotBlueprintBindableKey, "OculusTouch"));
+	AddKey(FKeyDetails(EKeys::OculusTouch_Left_Thumbstick_Click, LOCTEXT("OculusTouch_Left_Thumbstick_Click", "Oculus Touch (L) Thumbstick"), FKeyDetails::GamepadKey | FKeyDetails::NotBlueprintBindableKey, "OculusTouch"));
+	AddKey(FKeyDetails(EKeys::OculusTouch_Left_Thumbstick_Touch, LOCTEXT("OculusTouch_Left_Thumbstick_Touch", "Oculus Touch (L) Thumbstick Touch"), FKeyDetails::GamepadKey | FKeyDetails::NotBlueprintBindableKey, "OculusTouch"));
+	AddKey(FKeyDetails(EKeys::OculusTouch_Right_A_Click, LOCTEXT("OculusTouch_Right_A_Click", "Oculus Touch (R) A Press"), FKeyDetails::GamepadKey | FKeyDetails::NotBlueprintBindableKey, "OculusTouch"));
+	AddKey(FKeyDetails(EKeys::OculusTouch_Right_B_Click, LOCTEXT("OculusTouch_Right_B_Click", "Oculus Touch (R) B Press"), FKeyDetails::GamepadKey | FKeyDetails::NotBlueprintBindableKey, "OculusTouch"));
+	AddKey(FKeyDetails(EKeys::OculusTouch_Right_A_Touch, LOCTEXT("OculusTouch_Right_A_Touch", "Oculus Touch (R) A Touch"), FKeyDetails::GamepadKey | FKeyDetails::NotBlueprintBindableKey, "OculusTouch"));
+	AddKey(FKeyDetails(EKeys::OculusTouch_Right_B_Touch, LOCTEXT("OculusTouch_Right_B_Touch", "Oculus Touch (R) B Touch"), FKeyDetails::GamepadKey | FKeyDetails::NotBlueprintBindableKey, "OculusTouch"));
+	AddKey(FKeyDetails(EKeys::OculusTouch_Right_System_Click, LOCTEXT("OculusTouch_Right_System_Click", "Oculus Touch (R) System"), FKeyDetails::GamepadKey | FKeyDetails::NotBlueprintBindableKey, "OculusTouch"));
+	AddKey(FKeyDetails(EKeys::OculusTouch_Right_Grip_Click, LOCTEXT("OculusTouch_Right_Grip_Click", "Oculus Touch (R) Grip"), FKeyDetails::GamepadKey | FKeyDetails::NotBlueprintBindableKey, "OculusTouch"));
+	AddKey(FKeyDetails(EKeys::OculusTouch_Right_Grip_Axis, LOCTEXT("OculusTouch_Right_Grip_Axis", "Oculus Touch (R) Grip Axis"), FKeyDetails::GamepadKey | FKeyDetails::FloatAxis | FKeyDetails::NotBlueprintBindableKey, "OculusTouch"));
+	AddKey(FKeyDetails(EKeys::OculusTouch_Right_Trigger_Click, LOCTEXT("OculusTouch_Right_Trigger_Click", "Oculus Touch (R) Trigger"), FKeyDetails::GamepadKey | FKeyDetails::NotBlueprintBindableKey, "OculusTouch"));
+	AddKey(FKeyDetails(EKeys::OculusTouch_Right_Trigger_Axis, LOCTEXT("OculusTouch_Right_Trigger_Axis", "Oculus Touch (R) Trigger Axis"), FKeyDetails::GamepadKey | FKeyDetails::FloatAxis | FKeyDetails::NotBlueprintBindableKey, "OculusTouch"));
+	AddKey(FKeyDetails(EKeys::OculusTouch_Right_Trigger_Touch, LOCTEXT("OculusTouch_Right_Trigger_Touch", "Oculus Touch (R) Trigger Touch"), FKeyDetails::GamepadKey | FKeyDetails::NotBlueprintBindableKey, "OculusTouch"));
+	AddKey(FKeyDetails(EKeys::OculusTouch_Right_Thumbstick_X, LOCTEXT("OculusTouch_Right_Thumbstick_X", "Oculus Touch (R) Thumbstick X"), FKeyDetails::GamepadKey | FKeyDetails::FloatAxis | FKeyDetails::NotBlueprintBindableKey, "OculusTouch"));
+	AddKey(FKeyDetails(EKeys::OculusTouch_Right_Thumbstick_Y, LOCTEXT("OculusTouch_Right_Thumbstick_Y", "Oculus Touch (R) Thumbstick Y"), FKeyDetails::GamepadKey | FKeyDetails::FloatAxis | FKeyDetails::NotBlueprintBindableKey, "OculusTouch"));
+	AddKey(FKeyDetails(EKeys::OculusTouch_Right_Thumbstick_Click, LOCTEXT("OculusTouch_Right_Thumbstick_Click", "Oculus Touch (R) Thumbstick"), FKeyDetails::GamepadKey | FKeyDetails::NotBlueprintBindableKey, "OculusTouch"));
+	AddKey(FKeyDetails(EKeys::OculusTouch_Right_Thumbstick_Touch, LOCTEXT("OculusTouch_Right_Thumbstick_Touch", "Oculus Touch (R) Thumbstick Touch"), FKeyDetails::GamepadKey | FKeyDetails::NotBlueprintBindableKey, "OculusTouch"));
+
+	// Valve Index Controller
+	AddMenuCategoryDisplayInfo("ValveIndex", LOCTEXT("ValveIndexSubCategory", "Valve Index"), TEXT("GraphEditor.PadEvent_16x"));
+
+	AddKey(FKeyDetails(EKeys::ValveIndex_Left_A_Click, LOCTEXT("ValveIndex_Left_A_Click", "Valve Index (L) A Press"), FKeyDetails::GamepadKey | FKeyDetails::NotBlueprintBindableKey, "ValveIndex"));
+	AddKey(FKeyDetails(EKeys::ValveIndex_Left_B_Click, LOCTEXT("ValveIndex_Left_B_Click", "Valve Index (L) B Press"), FKeyDetails::GamepadKey | FKeyDetails::NotBlueprintBindableKey, "ValveIndex"));
+	AddKey(FKeyDetails(EKeys::ValveIndex_Left_A_Touch, LOCTEXT("ValveIndex_Left_A_Touch", "Valve Index (L) A Touch"), FKeyDetails::GamepadKey | FKeyDetails::NotBlueprintBindableKey, "ValveIndex"));
+	AddKey(FKeyDetails(EKeys::ValveIndex_Left_B_Touch, LOCTEXT("ValveIndex_Left_B_Touch", "Valve Index (L) B Touch"), FKeyDetails::GamepadKey | FKeyDetails::NotBlueprintBindableKey, "ValveIndex"));
+	AddKey(FKeyDetails(EKeys::ValveIndex_Left_System_Click, LOCTEXT("ValveIndex_Left_System_Click", "Valve Index (L) System"), FKeyDetails::GamepadKey | FKeyDetails::NotBlueprintBindableKey, "ValveIndex"));
+	AddKey(FKeyDetails(EKeys::ValveIndex_Left_System_Touch, LOCTEXT("ValveIndex_Left_System_Touch", "Valve Index (L) System Touch"), FKeyDetails::GamepadKey | FKeyDetails::NotBlueprintBindableKey, "ValveIndex"));
+	AddKey(FKeyDetails(EKeys::ValveIndex_Left_Grip_Click, LOCTEXT("ValveIndex_Left_Grip_Click", "Valve Index (L) Grip"), FKeyDetails::GamepadKey | FKeyDetails::NotBlueprintBindableKey, "ValveIndex"));
+	AddKey(FKeyDetails(EKeys::ValveIndex_Left_Grip_Axis, LOCTEXT("ValveIndex_Left_Grip_Axis", "Valve Index (L) Grip Axis"), FKeyDetails::GamepadKey | FKeyDetails::FloatAxis | FKeyDetails::NotBlueprintBindableKey, "ValveIndex"));
+	AddKey(FKeyDetails(EKeys::ValveIndex_Left_Grip_Force, LOCTEXT("ValveIndex_Left_Grip_Force", "Valve Index (L) Grip Force"), FKeyDetails::GamepadKey | FKeyDetails::FloatAxis | FKeyDetails::NotBlueprintBindableKey, "ValveIndex"));
+	AddKey(FKeyDetails(EKeys::ValveIndex_Left_Trigger_Click, LOCTEXT("ValveIndex_Left_Trigger_Click", "Valve Index (L) Trigger"), FKeyDetails::GamepadKey | FKeyDetails::NotBlueprintBindableKey, "ValveIndex"));
+	AddKey(FKeyDetails(EKeys::ValveIndex_Left_Trigger_Axis, LOCTEXT("ValveIndex_Left_Trigger_Axis", "Valve Index (L) Trigger Axis"), FKeyDetails::GamepadKey | FKeyDetails::FloatAxis | FKeyDetails::NotBlueprintBindableKey, "ValveIndex"));
+	AddKey(FKeyDetails(EKeys::ValveIndex_Left_Trigger_Touch, LOCTEXT("ValveIndex_Left_Trigger_Touch", "Valve Index (L) Trigger Touch"), FKeyDetails::GamepadKey | FKeyDetails::NotBlueprintBindableKey, "ValveIndex"));
+	AddKey(FKeyDetails(EKeys::ValveIndex_Left_Thumbstick_X, LOCTEXT("ValveIndex_Left_Thumbstick_X", "Valve Index (L) Thumbstick X"), FKeyDetails::GamepadKey | FKeyDetails::FloatAxis | FKeyDetails::NotBlueprintBindableKey, "ValveIndex"));
+	AddKey(FKeyDetails(EKeys::ValveIndex_Left_Thumbstick_Y, LOCTEXT("ValveIndex_Left_Thumbstick_Y", "Valve Index (L) Thumbstick Y"), FKeyDetails::GamepadKey | FKeyDetails::FloatAxis | FKeyDetails::NotBlueprintBindableKey, "ValveIndex"));
+	AddKey(FKeyDetails(EKeys::ValveIndex_Left_Thumbstick_Click, LOCTEXT("ValveIndex_Left_Thumbstick_Click", "Valve Index (L) Thumbstick"), FKeyDetails::GamepadKey | FKeyDetails::NotBlueprintBindableKey, "ValveIndex"));
+	AddKey(FKeyDetails(EKeys::ValveIndex_Left_Thumbstick_Touch, LOCTEXT("ValveIndex_Left_Thumbstick_Touch", "Valve Index (L) Thumbstick Touch"), FKeyDetails::GamepadKey | FKeyDetails::NotBlueprintBindableKey, "ValveIndex"));
+	AddKey(FKeyDetails(EKeys::ValveIndex_Left_Trackpad_X, LOCTEXT("ValveIndex_Left_Trackpad_X", "Valve Index (L) Trackpad X"), FKeyDetails::GamepadKey | FKeyDetails::FloatAxis | FKeyDetails::NotBlueprintBindableKey, "ValveIndex"));
+	AddKey(FKeyDetails(EKeys::ValveIndex_Left_Trackpad_Y, LOCTEXT("ValveIndex_Left_Trackpad_Y", "Valve Index (L) Trackpad Y"), FKeyDetails::GamepadKey | FKeyDetails::FloatAxis | FKeyDetails::NotBlueprintBindableKey, "ValveIndex"));
+	AddKey(FKeyDetails(EKeys::ValveIndex_Left_Trackpad_Click, LOCTEXT("ValveIndex_Left_Trackpad_Click", "Valve Index (L) Trackpad"), FKeyDetails::GamepadKey | FKeyDetails::NotBlueprintBindableKey, "ValveIndex"));
+	AddKey(FKeyDetails(EKeys::ValveIndex_Left_Trackpad_Force, LOCTEXT("ValveIndex_Left_Trackpad_Force", "Valve Index (L) Trackpad Force"), FKeyDetails::GamepadKey | FKeyDetails::FloatAxis | FKeyDetails::NotBlueprintBindableKey, "ValveIndex"));
+	AddKey(FKeyDetails(EKeys::ValveIndex_Left_Trackpad_Touch, LOCTEXT("ValveIndex_Left_Trackpad_Touch", "Valve Index (L) Trackpad Touch"), FKeyDetails::GamepadKey | FKeyDetails::FloatAxis | FKeyDetails::NotBlueprintBindableKey, "ValveIndex"));
+	AddKey(FKeyDetails(EKeys::ValveIndex_Right_A_Click, LOCTEXT("ValveIndex_Right_A_Click", "Valve Index (R) A Press"), FKeyDetails::GamepadKey | FKeyDetails::NotBlueprintBindableKey, "ValveIndex"));
+	AddKey(FKeyDetails(EKeys::ValveIndex_Right_B_Click, LOCTEXT("ValveIndex_Right_B_Click", "Valve Index (R) B Press"), FKeyDetails::GamepadKey | FKeyDetails::NotBlueprintBindableKey, "ValveIndex"));
+	AddKey(FKeyDetails(EKeys::ValveIndex_Right_A_Touch, LOCTEXT("ValveIndex_Right_A_Touch", "Valve Index (R) A Touch"), FKeyDetails::GamepadKey | FKeyDetails::NotBlueprintBindableKey, "ValveIndex"));
+	AddKey(FKeyDetails(EKeys::ValveIndex_Right_B_Touch, LOCTEXT("ValveIndex_Right_B_Touch", "Valve Index (R) B Touch"), FKeyDetails::GamepadKey | FKeyDetails::NotBlueprintBindableKey, "ValveIndex"));
+	AddKey(FKeyDetails(EKeys::ValveIndex_Right_System_Click, LOCTEXT("ValveIndex_Right_System_Click", "Valve Index (R) System"), FKeyDetails::GamepadKey | FKeyDetails::NotBlueprintBindableKey, "ValveIndex"));
+	AddKey(FKeyDetails(EKeys::ValveIndex_Right_System_Touch, LOCTEXT("ValveIndex_Right_System_Touch", "Valve Index (R) System Touch"), FKeyDetails::GamepadKey | FKeyDetails::NotBlueprintBindableKey, "ValveIndex"));
+	AddKey(FKeyDetails(EKeys::ValveIndex_Right_Grip_Click, LOCTEXT("ValveIndex_Right_Grip_Click", "Valve Index (R) Grip"), FKeyDetails::GamepadKey | FKeyDetails::NotBlueprintBindableKey, "ValveIndex"));
+	AddKey(FKeyDetails(EKeys::ValveIndex_Right_Grip_Axis, LOCTEXT("ValveIndex_Right_Grip_Axis", "Valve Index (R) Grip Axis"), FKeyDetails::GamepadKey | FKeyDetails::FloatAxis | FKeyDetails::NotBlueprintBindableKey, "ValveIndex"));
+	AddKey(FKeyDetails(EKeys::ValveIndex_Right_Grip_Force, LOCTEXT("ValveIndex_Right_Grip_Force", "Valve Index (R) Grip Force"), FKeyDetails::GamepadKey | FKeyDetails::FloatAxis | FKeyDetails::NotBlueprintBindableKey, "ValveIndex"));
+	AddKey(FKeyDetails(EKeys::ValveIndex_Right_Trigger_Click, LOCTEXT("ValveIndex_Right_Trigger_Click", "Valve Index (R) Trigger"), FKeyDetails::GamepadKey | FKeyDetails::NotBlueprintBindableKey, "ValveIndex"));
+	AddKey(FKeyDetails(EKeys::ValveIndex_Right_Trigger_Axis, LOCTEXT("ValveIndex_Right_Trigger_Axis", "Valve Index (R) Trigger Axis"), FKeyDetails::GamepadKey | FKeyDetails::FloatAxis | FKeyDetails::NotBlueprintBindableKey, "ValveIndex"));
+	AddKey(FKeyDetails(EKeys::ValveIndex_Right_Trigger_Touch, LOCTEXT("ValveIndex_Right_Trigger_Touch", "Valve Index (R) Trigger Touch"), FKeyDetails::GamepadKey | FKeyDetails::NotBlueprintBindableKey, "ValveIndex"));
+	AddKey(FKeyDetails(EKeys::ValveIndex_Right_Thumbstick_X, LOCTEXT("ValveIndex_Right_Thumbstick_X", "Valve Index (R) Thumbstick X"), FKeyDetails::GamepadKey | FKeyDetails::FloatAxis | FKeyDetails::NotBlueprintBindableKey, "ValveIndex"));
+	AddKey(FKeyDetails(EKeys::ValveIndex_Right_Thumbstick_Y, LOCTEXT("ValveIndex_Right_Thumbstick_Y", "Valve Index (R) Thumbstick Y"), FKeyDetails::GamepadKey | FKeyDetails::FloatAxis | FKeyDetails::NotBlueprintBindableKey, "ValveIndex"));
+	AddKey(FKeyDetails(EKeys::ValveIndex_Right_Thumbstick_Click, LOCTEXT("ValveIndex_Right_Thumbstick_Click", "Valve Index (R) Thumbstick"), FKeyDetails::GamepadKey | FKeyDetails::NotBlueprintBindableKey, "ValveIndex"));
+	AddKey(FKeyDetails(EKeys::ValveIndex_Right_Thumbstick_Touch, LOCTEXT("ValveIndex_Right_Thumbstick_Touch", "Valve Index (R) Thumbstick Touch"), FKeyDetails::GamepadKey | FKeyDetails::NotBlueprintBindableKey, "ValveIndex"));
+	AddKey(FKeyDetails(EKeys::ValveIndex_Right_Trackpad_X, LOCTEXT("ValveIndex_Right_Trackpad_X", "Valve Index (R) Trackpad X"), FKeyDetails::GamepadKey | FKeyDetails::FloatAxis | FKeyDetails::NotBlueprintBindableKey, "ValveIndex"));
+	AddKey(FKeyDetails(EKeys::ValveIndex_Right_Trackpad_Y, LOCTEXT("ValveIndex_Right_Trackpad_Y", "Valve Index (R) Trackpad Y"), FKeyDetails::GamepadKey | FKeyDetails::FloatAxis | FKeyDetails::NotBlueprintBindableKey, "ValveIndex"));
+	AddKey(FKeyDetails(EKeys::ValveIndex_Right_Trackpad_Click, LOCTEXT("ValveIndex_Right_Trackpad_Click", "Valve Index (R) Trackpad"), FKeyDetails::GamepadKey | FKeyDetails::NotBlueprintBindableKey, "ValveIndex"));
+	AddKey(FKeyDetails(EKeys::ValveIndex_Right_Trackpad_Force, LOCTEXT("ValveIndex_Right_Trackpad_Force", "Valve Index (R) Trackpad Force"), FKeyDetails::GamepadKey | FKeyDetails::FloatAxis | FKeyDetails::NotBlueprintBindableKey, "ValveIndex"));
+	AddKey(FKeyDetails(EKeys::ValveIndex_Right_Trackpad_Touch, LOCTEXT("ValveIndex_Right_Trackpad_Touch", "Valve Index (R) Trackpad Touch"), FKeyDetails::GamepadKey | FKeyDetails::NotBlueprintBindableKey, "ValveIndex"));
+
 	// Initialize the input key manager.  This will cause any additional OEM keys to get added
 	FInputKeyManager::Get();
 }
@@ -971,8 +1276,8 @@ FText EKeys::GetGamepadDisplayName(const FKey Key)
 }
 PRAGMA_ENABLE_OPTIMIZATION
 
-bool FKey::IsValid() const 
-{ 
+bool FKey::IsValid() const
+{
 	if (KeyName != NAME_None)
 	{
 		ConditionalLookupKeyDetails();
@@ -983,7 +1288,7 @@ bool FKey::IsValid() const
 
 FString FKey::ToString() const
 {
-	return KeyName.ToString(); 
+	return KeyName.ToString();
 }
 
 FName FKey::GetFName() const
@@ -991,7 +1296,7 @@ FName FKey::GetFName() const
 	return KeyName;
 }
 
-bool FKey::IsModifierKey() const 
+bool FKey::IsModifierKey() const
 {
 	ConditionalLookupKeyDetails();
 	return (KeyDetails.IsValid() ? KeyDetails->IsModifierKey() : false);
@@ -1039,6 +1344,12 @@ bool FKey::ShouldUpdateAxisWithoutSamples() const
 	return (KeyDetails.IsValid() ? KeyDetails->ShouldUpdateAxisWithoutSamples() : false);
 }
 
+bool FKey::IsBindableToActions() const
+{
+	ConditionalLookupKeyDetails();
+	return (KeyDetails.IsValid() ? KeyDetails->IsBindableToActions() : false);
+}
+
 FText FKeyDetails::GetDisplayName(const bool bLongDisplayName) const
 {
 	return ((bLongDisplayName || !ShortDisplayName.IsSet()) ? LongDisplayName.Get() : ShortDisplayName.Get());
@@ -1143,7 +1454,7 @@ void FKey::ResetKey()
 TSharedPtr<FInputKeyManager> FInputKeyManager::Instance;
 
 /**
- * Returns the instance of the input key manager                   
+ * Returns the instance of the input key manager
  */
 FInputKeyManager& FInputKeyManager::Get()
 {
@@ -1166,7 +1477,7 @@ void FInputKeyManager::InitKeyMappings()
 	for (uint32 Idx=0; Idx<KeyMapSize; ++Idx)
 	{
 		FKey Key(*KeyNames[Idx]);
-		
+
 		if (!Key.IsValid())
 		{
 			EKeys::AddKey(FKeyDetails(Key, Key.GetDisplayName()));
@@ -1203,7 +1514,7 @@ FKey FInputKeyManager::GetKeyFromCodes( const uint32 KeyCode, const uint32 CharC
 	if (KeyPtr == nullptr && CharCode > 32)
 	{
 		FKey NewKey(*FString::Printf(TEXT("%s%d"), FKey::SyntheticCharPrefix, CharCode));
-		EKeys::AddKey(FKeyDetails(NewKey, FText::AsCultureInvariant(FString::Chr(CharCode)), FKeyDetails::NotBlueprintBindableKey));
+		EKeys::AddKey(FKeyDetails(NewKey, FText::AsCultureInvariant(FString::Chr(CharCode)), FKeyDetails::NotBlueprintBindableKey | FKeyDetails::NotActionBindableKey));
 		const_cast<FInputKeyManager*>(this)->KeyMapCharToEnum.Add(CharCode, NewKey);
 		return NewKey;
 	}
