Auto-generate bindings with an option to enable a "developer" mode
Move directional keys to input core for compatibility
Make the project name more unique to avoid manifest conflicts
Fix bindings being overwritten every time a project is opened

diff --git a/Engine/Plugins/Runtime/Oculus/OculusVR/Source/OculusInput/Private/OculusInputState.h b/Engine/Plugins/Runtime/Oculus/OculusVR/Source/OculusInput/Private/OculusInputState.h
--- a/Engine/Plugins/Runtime/Oculus/OculusVR/Source/OculusInput/Private/OculusInputState.h
+++ b/Engine/Plugins/Runtime/Oculus/OculusVR/Source/OculusInput/Private/OculusInputState.h
@@ -299,21 +299,25 @@ struct FOculusTouchControllerState
 			Button.NextRepeatTime = 0.0;
 		}
 
-		Buttons[(int32)EOculusTouchControllerButton::Trigger].Key = (Hand == EControllerHand::Left) ? EKeys::OculusTouch_Left_Trigger_Click.GetFName() : EKeys::OculusTouch_Right_Trigger_Click.GetFName();
-		Buttons[(int32)EOculusTouchControllerButton::Grip].Key = (Hand == EControllerHand::Left) ? EKeys::OculusTouch_Left_Grip_Click.GetFName() : EKeys::OculusTouch_Right_Grip_Click.GetFName();
-		Buttons[(int32)EOculusTouchControllerButton::Thumbstick].Key = (Hand == EControllerHand::Left) ? EKeys::OculusTouch_Left_Thumbstick_Click.GetFName() : EKeys::OculusTouch_Right_Thumbstick_Click.GetFName();
-		Buttons[(int32)EOculusTouchControllerButton::XA].Key = (Hand == EControllerHand::Left) ? EKeys::OculusTouch_Left_X_Click.GetFName() : EKeys::OculusTouch_Right_A_Click.GetFName();
-		Buttons[(int32)EOculusTouchControllerButton::YB].Key = (Hand == EControllerHand::Left) ? EKeys::OculusTouch_Left_Y_Click.GetFName() : EKeys::OculusTouch_Right_B_Click.GetFName();
+		Buttons[ (int32)EOculusTouchControllerButton::Trigger ].Key = (Hand == EControllerHand::Left) ? EKeys::OculusTouch_Left_Trigger_Click.GetFName() : EKeys::OculusTouch_Right_Trigger_Click.GetFName();
+		Buttons[ (int32)EOculusTouchControllerButton::Grip ].Key = (Hand == EControllerHand::Left) ? EKeys::OculusTouch_Left_Grip_Click.GetFName() : EKeys::OculusTouch_Right_Grip_Click.GetFName();
+		Buttons[ (int32)EOculusTouchControllerButton::Thumbstick ].Key = (Hand == EControllerHand::Left) ? EKeys::OculusTouch_Left_Thumbstick_Click.GetFName() : EKeys::OculusTouch_Right_Thumbstick_Click.GetFName();
+		Buttons[ (int32)EOculusTouchControllerButton::XA ].Key = (Hand == EControllerHand::Left) ? EKeys::OculusTouch_Left_X_Click.GetFName() : EKeys::OculusTouch_Right_A_Click.GetFName();
+		Buttons[ (int32)EOculusTouchControllerButton::YB ].Key = (Hand == EControllerHand::Left) ? EKeys::OculusTouch_Left_Y_Click.GetFName() : EKeys::OculusTouch_Right_B_Click.GetFName();
+		Buttons[ (int32)EOculusTouchControllerButton::Thumbstick_Up ].Key = (Hand == EControllerHand::Left) ? EKeys::OculusTouch_Left_Thumbstick_Up.GetFName() : EKeys::OculusTouch_Right_Thumbstick_Up.GetFName();
+		Buttons[ (int32)EOculusTouchControllerButton::Thumbstick_Down ].Key = (Hand == EControllerHand::Left) ? EKeys::OculusTouch_Left_Thumbstick_Down.GetFName() : EKeys::OculusTouch_Right_Thumbstick_Down.GetFName();
+		Buttons[ (int32)EOculusTouchControllerButton::Thumbstick_Left ].Key = (Hand == EControllerHand::Left) ? EKeys::OculusTouch_Left_Thumbstick_Left.GetFName() : EKeys::OculusTouch_Right_Thumbstick_Left.GetFName();
+		Buttons[ (int32)EOculusTouchControllerButton::Thumbstick_Right ].Key = (Hand == EControllerHand::Left) ? EKeys::OculusTouch_Left_Thumbstick_Right.GetFName() : EKeys::OculusTouch_Right_Thumbstick_Right.GetFName();
 
 		Buttons[ (int32)EOculusTouchControllerButton::Trigger ].EmulatedKey = ( Hand == EControllerHand::Left ) ? FGamepadKeyNames::MotionController_Left_Trigger : FGamepadKeyNames::MotionController_Right_Trigger;
 		Buttons[ (int32)EOculusTouchControllerButton::Grip ].EmulatedKey = ( Hand == EControllerHand::Left ) ? FGamepadKeyNames::MotionController_Left_Grip1 : FGamepadKeyNames::MotionController_Right_Grip1;
 		Buttons[ (int32)EOculusTouchControllerButton::Thumbstick ].EmulatedKey = ( Hand == EControllerHand::Left ) ? FGamepadKeyNames::MotionController_Left_Thumbstick : FGamepadKeyNames::MotionController_Right_Thumbstick;
 		Buttons[ (int32)EOculusTouchControllerButton::XA ].EmulatedKey = (Hand == EControllerHand::Left) ? FGamepadKeyNames::MotionController_Left_FaceButton1 : FGamepadKeyNames::MotionController_Right_FaceButton1;
 		Buttons[ (int32)EOculusTouchControllerButton::YB ].EmulatedKey = (Hand == EControllerHand::Left) ? FGamepadKeyNames::MotionController_Left_FaceButton2 : FGamepadKeyNames::MotionController_Right_FaceButton2;
-		Buttons[ (int32)EOculusTouchControllerButton::Thumbstick_Up].Key = (Hand == EControllerHand::Left) ? FGamepadKeyNames::MotionController_Left_Thumbstick_Up : FGamepadKeyNames::MotionController_Right_Thumbstick_Up;
-		Buttons[ (int32)EOculusTouchControllerButton::Thumbstick_Down].Key = (Hand == EControllerHand::Left) ? FGamepadKeyNames::MotionController_Left_Thumbstick_Down : FGamepadKeyNames::MotionController_Right_Thumbstick_Down;
-		Buttons[ (int32)EOculusTouchControllerButton::Thumbstick_Left].Key = (Hand == EControllerHand::Left) ? FGamepadKeyNames::MotionController_Left_Thumbstick_Left : FGamepadKeyNames::MotionController_Right_Thumbstick_Left;
-		Buttons[ (int32)EOculusTouchControllerButton::Thumbstick_Right].Key = (Hand == EControllerHand::Left) ? FGamepadKeyNames::MotionController_Left_Thumbstick_Right : FGamepadKeyNames::MotionController_Right_Thumbstick_Right;
+		Buttons[ (int32)EOculusTouchControllerButton::Thumbstick_Up].EmulatedKey = (Hand == EControllerHand::Left) ? FGamepadKeyNames::MotionController_Left_Thumbstick_Up : FGamepadKeyNames::MotionController_Right_Thumbstick_Up;
+		Buttons[ (int32)EOculusTouchControllerButton::Thumbstick_Down].EmulatedKey = (Hand == EControllerHand::Left) ? FGamepadKeyNames::MotionController_Left_Thumbstick_Down : FGamepadKeyNames::MotionController_Right_Thumbstick_Down;
+		Buttons[ (int32)EOculusTouchControllerButton::Thumbstick_Left].EmulatedKey = (Hand == EControllerHand::Left) ? FGamepadKeyNames::MotionController_Left_Thumbstick_Left : FGamepadKeyNames::MotionController_Right_Thumbstick_Left;
+		Buttons[ (int32)EOculusTouchControllerButton::Thumbstick_Right].EmulatedKey = (Hand == EControllerHand::Left) ? FGamepadKeyNames::MotionController_Left_Thumbstick_Right : FGamepadKeyNames::MotionController_Right_Thumbstick_Right;
 
 		Buttons[ (int32)EOculusTouchControllerButton::Menu ].Key = (Hand == EControllerHand::Left) ? FGamepadKeyNames::SpecialLeft : FGamepadKeyNames::SpecialRight;
 
diff --git a/Engine/Plugins/Runtime/OpenXR/Source/OpenXRInput/Private/OpenXRInput.cpp b/Engine/Plugins/Runtime/OpenXR/Source/OpenXRInput/Private/OpenXRInput.cpp
--- a/Engine/Plugins/Runtime/OpenXR/Source/OpenXRInput/Private/OpenXRInput.cpp
+++ b/Engine/Plugins/Runtime/OpenXR/Source/OpenXRInput/Private/OpenXRInput.cpp
@@ -299,7 +299,7 @@ int32 FOpenXRInputPlugin::FOpenXRInput::SuggestBindings(XrInstance Instance, FOp
 	// Add suggested bindings for every mapping
 	for (const T& InputKey : Mappings)
 	{
-		// Key names that are parseable into an OpenXR path have exeactly 4 components
+		// Key names that are parseable into an OpenXR path have exactly 4 tokens
 		TArray<FString> Tokens;
 		if (InputKey.Key.ToString().ParseIntoArray(Tokens, TEXT("_")) != 4)
 		{
@@ -320,23 +320,35 @@ int32 FOpenXRInputPlugin::FOpenXRInput::SuggestBindings(XrInstance Instance, FOp
 			Action.KeyMap.Add(Key, InputKey.Key.GetFName());
 
 			// Add the input we want to query with grip being defined as "squeeze" in OpenXR
-			FString InputString = Tokens[2].ToLower();
-			if (InputString == "grip")
+			FString Identifier = Tokens[2].ToLower();
+			if (Identifier == "grip")
 			{
-				InputString = "squeeze";
+				Identifier = "squeeze";
 			}
-			Path += "/input/" + InputString;
+			Path += "/input/" + Identifier;
 
 			// Add the data we want to query, we'll skip this for trigger/squeeze "click" actions to allow
 			// certain profiles that don't have "click" data to threshold the "value" data instead
-			FString DataString = Tokens[3].ToLower();
-			if (DataString == "axis")
+			FString Component = Tokens[3].ToLower();
+			if (Component == "axis")
 			{
-				DataString = "value";
+				Path += "/value";
 			}
-			if ((InputString != "trigger" && InputString != "squeeze") || DataString != "click")
+			else if (Component == "click")
 			{
-				Path += "/" + DataString;
+				if (Identifier != "trigger" && Identifier != "squeeze")
+				{
+					Path += "/click";
+				}
+			}
+			else if (Component == "touch" || Component == "x" || Component == "y")
+			{
+				Path += "/" + Component;
+			}
+			else
+			{
+				// Unrecognized data
+				continue;
 			}
 
 			// Add the binding to the profile
diff --git a/Engine/Plugins/Runtime/Steam/SteamVR/Source/SteamVREditor/Private/SteamVREditor.cpp b/Engine/Plugins/Runtime/Steam/SteamVR/Source/SteamVREditor/Private/SteamVREditor.cpp
--- a/Engine/Plugins/Runtime/Steam/SteamVR/Source/SteamVREditor/Private/SteamVREditor.cpp
+++ b/Engine/Plugins/Runtime/Steam/SteamVR/Source/SteamVREditor/Private/SteamVREditor.cpp
@@ -8,6 +8,7 @@
 #include "SteamVREditorStyle.h"
 #include "SteamVREditorCommands.h"
 #include "SteamVRControllerKeys.h"
+#include "SteamVRInputSettings.h"
 #include "LevelEditor.h"
 
 static const FName SteamVREditorTabName("SteamVREditor");
@@ -60,11 +61,14 @@ void FSteamVREditorModule::StartupModule()
 		FCanExecuteAction());
 
 	FLevelEditorModule& LevelEditorModule = FModuleManager::LoadModuleChecked<FLevelEditorModule>("LevelEditor");
-
-	TSharedPtr<FExtender> MenuExtender = MakeShareable(new FExtender());
-	MenuExtender->AddMenuExtension("EditLocalTabSpawners", EExtensionHook::After, PluginCommands, FMenuExtensionDelegate::CreateRaw(this, &FSteamVREditorModule::AddMenuExtension));
-
-	LevelEditorModule.GetMenuExtensibilityManager()->AddExtender(MenuExtender);
+	const USteamVRInputSettings* Settings = GetDefault<USteamVRInputSettings>();
+	if (Settings->bDeveloperMode)
+	{
+		TSharedPtr<FExtender> ToolbarExtender = MakeShareable(new FExtender);
+		ToolbarExtender->AddToolBarExtension("Settings", EExtensionHook::After, PluginCommands, FToolBarExtensionDelegate::CreateRaw(this, &FSteamVREditorModule::AddToolbarExtension));
+		
+		LevelEditorModule.GetToolBarExtensibilityManager()->AddExtender(ToolbarExtender);
+	}
 }
 
 void FSteamVREditorModule::ShutdownModule()
@@ -272,26 +276,35 @@ bool FSteamVREditorModule::AddUniqueActionMapping(TArray<FInputActionKeyMapping>
 }
 
 void FSteamVREditorModule::AddMenuExtension(FMenuBuilder& Builder)
+{
+	Builder.AddMenuEntry(FSteamVREditorCommands::Get().PluginAction);
+}
+
+void FSteamVREditorModule::AddToolbarExtension(FToolBarBuilder& Builder)
 {
 	FSteamVREditorStyle MenuStyle = FSteamVREditorStyle();
 	MenuStyle.Initialize();
 
-	Builder.AddSubMenu(
+	Builder.AddComboButton(
+		FUIAction(FExecuteAction::CreateRaw(this, &FSteamVREditorModule::PluginButtonClicked)),
+		FOnGetContent::CreateRaw(this, &FSteamVREditorModule::FillComboButton, PluginCommands),
 		LOCTEXT("SteamVRInputBtn", "SteamVR Input"),
 		LOCTEXT("SteamVRInputBtnTootlip", "SteamVR Input"),
-		FNewMenuDelegate::CreateRaw(this, &FSteamVREditorModule::MakeMenu),
-		false, FSlateIcon(FSteamVREditorStyle::GetStyleSetName(), "SteamVREditor.PluginAction")
+		FSlateIcon(FSteamVREditorStyle::GetStyleSetName(), "SteamVREditor.PluginAction")
 	);
 }
 
-void FSteamVREditorModule::MakeMenu(FMenuBuilder& MenuBuilder)
+TSharedRef<SWidget> FSteamVREditorModule::FillComboButton(TSharedPtr<class FUICommandList> Commands)
 {
-	//MenuBuilder.AddMenuEntry(FSteamVREditorCommands::Get().PluginAction);
+	FMenuBuilder MenuBuilder(true, Commands);
+
 	MenuBuilder.AddMenuEntry(FSteamVREditorCommands::Get().JsonActionManifest, NAME_None, TAttribute<FText>(), TAttribute<FText>(), FSlateIcon(FSteamVREditorStyle::GetStyleSetName(), "SteamVREditor.JsonActionManifest"));
 	MenuBuilder.AddMenuEntry(FSteamVREditorCommands::Get().JsonControllerBindings, NAME_None, TAttribute<FText>(), TAttribute<FText>(), FSlateIcon(FSteamVREditorStyle::GetStyleSetName(), "SteamVREditor.JsonControllerBindings"));
 	MenuBuilder.AddMenuEntry(FSteamVREditorCommands::Get().ReloadActionManifest, NAME_None, TAttribute<FText>(), TAttribute<FText>(), FSlateIcon(FSteamVREditorStyle::GetStyleSetName(), "SteamVREditor.ReloadActionManifest"));
 	MenuBuilder.AddMenuEntry(FSteamVREditorCommands::Get().LaunchBindingsURL, NAME_None, TAttribute<FText>(), TAttribute<FText>(), FSlateIcon(FSteamVREditorStyle::GetStyleSetName(), "SteamVREditor.LaunchBindingsURL"));
 	//MenuBuilder.AddMenuEntry(FSteamVREditorCommands::Get().AddSampleInputs, NAME_None, TAttribute<FText>(), TAttribute<FText>(), FSlateIcon(FSteamVREditorStyle::GetStyleSetName(), "SteamVREditor.AddSampleInputs"));
+
+	return MenuBuilder.MakeWidget();
 }
 
 #undef LOCTEXT_NAMESPACE
diff --git a/Engine/Plugins/Runtime/Steam/SteamVR/Source/SteamVREditor/Private/SteamVREditorStyle.cpp b/Engine/Plugins/Runtime/Steam/SteamVR/Source/SteamVREditor/Private/SteamVREditorStyle.cpp
--- a/Engine/Plugins/Runtime/Steam/SteamVR/Source/SteamVREditor/Private/SteamVREditorStyle.cpp
+++ b/Engine/Plugins/Runtime/Steam/SteamVR/Source/SteamVREditor/Private/SteamVREditorStyle.cpp
@@ -43,7 +43,7 @@ TSharedRef< FSlateStyleSet > FSteamVREditorStyle::Create()
 	TSharedRef< FSlateStyleSet > Style = MakeShareable(new FSlateStyleSet("SteamVREditorStyle"));
 	Style->SetContentRoot(IPluginManager::Get().FindPlugin("SteamVR")->GetBaseDir() / TEXT("Resources"));
 
-	Style->Set("SteamVREditor.PluginAction", new IMAGE_BRUSH(TEXT("Icon_PluginAction_16x"), Icon16x16));
+	Style->Set("SteamVREditor.PluginAction", new IMAGE_BRUSH(TEXT("Icon_PluginAction"), Icon40x40));
 	Style->Set("SteamVREditor.JsonActionManifest", new IMAGE_BRUSH(TEXT("Icon_regenerate_action_manifest"), Icon16x16));
 	Style->Set("SteamVREditor.JsonControllerBindings", new IMAGE_BRUSH(TEXT("Icon_regenerate_controller_bindings"), Icon16x16));
 	Style->Set("SteamVREditor.ReloadActionManifest", new IMAGE_BRUSH(TEXT("Icon_reload_action_manifest"), Icon16x16));
diff --git a/Engine/Plugins/Runtime/Steam/SteamVR/Source/SteamVREditor/Public/SteamVREditor.h b/Engine/Plugins/Runtime/Steam/SteamVR/Source/SteamVREditor/Public/SteamVREditor.h
--- a/Engine/Plugins/Runtime/Steam/SteamVR/Source/SteamVREditor/Public/SteamVREditor.h
+++ b/Engine/Plugins/Runtime/Steam/SteamVR/Source/SteamVREditor/Public/SteamVREditor.h
@@ -29,11 +29,12 @@ public:
 
 private:
 
+	void AddToolbarExtension(FToolBarBuilder& Builder);
 	void AddMenuExtension(FMenuBuilder& Builder);
 	bool AddUniqueAxisMapping(TArray<FInputAxisKeyMapping> ExistingAxisKeys, UInputSettings* InputSettings, FName ActionName, FKey ActionKey);
 	bool AddUniqueActionMapping(TArray<FInputActionKeyMapping> ExistingActionKeys, UInputSettings* InputSettings, FName ActionName, FKey ActionKey);
 
 	TSharedPtr<class FUICommandList> PluginCommands;
-	void MakeMenu(FMenuBuilder& MenuBuilder);
+	TSharedRef<SWidget> FillComboButton(TSharedPtr<class FUICommandList> Commands);
 
 };
diff --git a/Engine/Plugins/Runtime/Steam/SteamVR/Source/SteamVRInput/Private/SteamVRInput.cpp b/Engine/Plugins/Runtime/Steam/SteamVR/Source/SteamVRInput/Private/SteamVRInput.cpp
--- a/Engine/Plugins/Runtime/Steam/SteamVR/Source/SteamVRInput/Private/SteamVRInput.cpp
+++ b/Engine/Plugins/Runtime/Steam/SteamVR/Source/SteamVRInput/Private/SteamVRInput.cpp
@@ -1,22 +1,64 @@
 // Copyright 1998-2019 Epic Games, Inc. All Rights Reserved.
 
 #include "SteamVRInput.h"
+#include "SteamVRInputSettings.h"
 #include "CoreMinimal.h"
 #include "Runtime/Core/Public/Misc/Paths.h"
 #include "openvr.h"
 #include "Modules/ModuleManager.h"
 #include "Interfaces/IPluginManager.h"
 
+#if WITH_EDITOR
+#include "Editor.h"
+#include "ISettingsModule.h"
+#include "ISettingsSection.h"
+#endif
+
 #define LOCTEXT_NAMESPACE "FSteamVRInputModule"
 
 DEFINE_LOG_CATEGORY_STATIC(LogSteamInput, Log, All);
 
+#if WITH_EDITOR
+void FSteamVRInputModule::AddEditorSettings()
+{
+	ISettingsModule* SettingsModule = FModuleManager::GetModulePtr<ISettingsModule>("Settings");
+
+	// While this should usually be true, it's not guaranteed that the settings module will be loaded in the editor.
+	// UBT allows setting bBuildDeveloperTools to false while bBuildEditor can be true.
+	// The former option indirectly controls loading of the "Settings" module.
+	if (SettingsModule)
+	{
+		SettingsModule->RegisterSettings("Project", "Plugins", "SteamVR Input",
+			LOCTEXT("SteamVRInputSettingsName", "SteamVR Input"),
+			LOCTEXT("SteamVRInputSettingsDescription", "Configure SteamVR input support."),
+			GetMutableDefault<USteamVRInputSettings>()
+		);
+	}
+}
+
+void FSteamVRInputModule::RemoveEditorSettings()
+{
+	ISettingsModule* SettingsModule = FModuleManager::GetModulePtr<ISettingsModule>("Settings");
+	if (SettingsModule)
+	{
+		SettingsModule->UnregisterSettings("Project", "Plugins", "SteamVR Input");
+	}
+}
+#endif
+
 void FSteamVRInputModule::StartupModule()
 {
+#if WITH_EDITOR
+	// We don't quite have control of when the "Settings" module is loaded, so we'll wait until PostEngineInit to register settings.
+	FCoreDelegates::OnPostEngineInit.AddRaw(this, &FSteamVRInputModule::AddEditorSettings);
+#endif // WITH_EDITOR
 }
 
 void FSteamVRInputModule::ShutdownModule()
 {
+#if WITH_EDITOR
+	RemoveEditorSettings();
+#endif
 }
 
 #undef LOCTEXT_NAMESPACE
diff --git a/Engine/Plugins/Runtime/Steam/SteamVR/Source/SteamVRInput/Public/SteamVRInput.h b/Engine/Plugins/Runtime/Steam/SteamVR/Source/SteamVRInput/Public/SteamVRInput.h
--- a/Engine/Plugins/Runtime/Steam/SteamVR/Source/SteamVRInput/Public/SteamVRInput.h
+++ b/Engine/Plugins/Runtime/Steam/SteamVR/Source/SteamVRInput/Public/SteamVRInput.h
@@ -12,6 +12,11 @@ public:
 	virtual void StartupModule() override;
 	virtual void ShutdownModule() override;
 
+#if WITH_EDITOR
+	void AddEditorSettings();
+	void RemoveEditorSettings();
+#endif
+
 	static inline IModuleInterface& Get()
 	{
 		return FModuleManager::LoadModuleChecked<IModuleInterface>("SteamVRInput");
diff --git a/Engine/Plugins/Runtime/Steam/SteamVR/Source/SteamVRInput/Public/SteamVRInputSettings.h b/Engine/Plugins/Runtime/Steam/SteamVR/Source/SteamVRInput/Public/SteamVRInputSettings.h
new file mode 100644
--- /dev/null
+++ b/Engine/Plugins/Runtime/Steam/SteamVR/Source/SteamVRInput/Public/SteamVRInputSettings.h
@@ -0,0 +1,17 @@
+// Copyright 1998-2019 Epic Games, Inc. All Rights Reserved.
+
+#include "CoreMinimal.h"
+#include "SteamVRInputSettings.generated.h"
+
+/**
+* Implements the settings for SteamVR input.
+*/
+UCLASS(config = Editor, defaultconfig)
+class STEAMVRINPUT_API USteamVRInputSettings : public UObject
+{
+public:
+	GENERATED_BODY()
+
+	UPROPERTY(GlobalConfig, EditAnywhere, Category = "General", Meta = (DisplayName = "Enable SteamVR Input developer mode", Tooltip = "Disables automatic manifest generation and enables toolbar icon", ConfigRestartRequired = true))
+	bool bDeveloperMode = false;
+};
diff --git a/Engine/Plugins/Runtime/Steam/SteamVR/Source/SteamVRInput/SteamVRInput.Build.cs b/Engine/Plugins/Runtime/Steam/SteamVR/Source/SteamVRInput/SteamVRInput.Build.cs
--- a/Engine/Plugins/Runtime/Steam/SteamVR/Source/SteamVRInput/SteamVRInput.Build.cs
+++ b/Engine/Plugins/Runtime/Steam/SteamVR/Source/SteamVRInput/SteamVRInput.Build.cs
@@ -12,14 +12,20 @@ public class SteamVRInput : ModuleRules
         PrivateIncludePaths.Add(Path.Combine(ModuleDirectory, "Private"));
         PublicDependencyModuleNames.AddRange(
 			new string[]
-			{
-				"Core",
+            {
+                "Core",
+                "CoreUObject",
                 "Projects",
                 "OpenVR",
                 "SteamVR"
             }
 			);
 
+        if (Target.bBuildEditor == true)
+        {
+            PrivateDependencyModuleNames.Add("UnrealEd");
+        }
+
         RuntimeDependencies.Add("$(ProjectDir)/Config/SteamVRBindings/...");
     }
 }
diff --git a/Engine/Plugins/Runtime/Steam/SteamVR/Source/SteamVRInputDevice/Private/SteamVRInputDevice.cpp b/Engine/Plugins/Runtime/Steam/SteamVR/Source/SteamVRInputDevice/Private/SteamVRInputDevice.cpp
--- a/Engine/Plugins/Runtime/Steam/SteamVR/Source/SteamVRInputDevice/Private/SteamVRInputDevice.cpp
+++ b/Engine/Plugins/Runtime/Steam/SteamVR/Source/SteamVRInputDevice/Private/SteamVRInputDevice.cpp
@@ -1,6 +1,7 @@
 // Copyright 1998-2019 Epic Games, Inc. All Rights Reserved.
 
 #include "SteamVRInputDevice.h"
+#include "SteamVRInputSettings.h"
 #include "Runtime/ApplicationCore/Public/GenericPlatform/IInputInterface.h"
 #include "HAL/FileManagerGeneric.h"
 #include "Misc/FileHelper.h"
@@ -83,18 +84,22 @@ FSteamVRInputDevice::FSteamVRInputDevice(const TSharedRef<FGenericApplicationMes
 	GenerateActionManifest();
 
 	// Auto-enable SteamVR Input Developer Mode (reload hmd module)
-	if (VRSettings() != nullptr)
+	const USteamVRInputSettings* Settings = GetDefault<USteamVRInputSettings>();
+	if (Settings->bDeveloperMode && VRSettings() != nullptr)
 	{
 		EVRInitError SteamVRInitError = VRInitError_Driver_NotLoaded;
 		IVRSystem* SteamVRSystem = vr::VR_Init(&SteamVRInitError, vr::VRApplication_Overlay);
 
-		EVRSettingsError BindingFlagError = VRSettingsError_None;
-		VRSettings()->SetBool(k_pch_SteamVR_Section, k_pch_SteamVR_DebugInputBinding, true, &BindingFlagError);
-		UE_LOG(LogSteamVRInputDevice, Display, TEXT("[STEAMVR INPUT] Enable SteamVR Input Developer Mode: %s"), *FString(UTF8_TO_TCHAR(VRSettings()->GetSettingsErrorNameFromEnum(BindingFlagError))));
-		VRSettings()->SetBool(k_pch_SteamVR_Section, k_pch_SteamVR_DebugInput, true, &BindingFlagError);
-		UE_LOG(LogSteamVRInputDevice, Display, TEXT("[STEAMVR INPUT] Enable SteamVR Debug Input: %s"), *FString(UTF8_TO_TCHAR(VRSettings()->GetSettingsErrorNameFromEnum(BindingFlagError))));
+		if (SteamVRInitError == VRInitError_None)
+		{
+			EVRSettingsError BindingFlagError = VRSettingsError_None;
+			VRSettings()->SetBool(k_pch_SteamVR_Section, k_pch_SteamVR_DebugInputBinding, true, &BindingFlagError);
+			UE_LOG(LogSteamVRInputDevice, Display, TEXT("[STEAMVR INPUT] Enable SteamVR Input Developer Mode: %s"), *FString(UTF8_TO_TCHAR(VRSettings()->GetSettingsErrorNameFromEnum(BindingFlagError))));
+			VRSettings()->SetBool(k_pch_SteamVR_Section, k_pch_SteamVR_DebugInput, true, &BindingFlagError);
+			UE_LOG(LogSteamVRInputDevice, Display, TEXT("[STEAMVR INPUT] Enable SteamVR Debug Input: %s"), *FString(UTF8_TO_TCHAR(VRSettings()->GetSettingsErrorNameFromEnum(BindingFlagError))));
 
-		VR_Shutdown();
+			VR_Shutdown();
+		}
 	}
 #endif
 
@@ -1139,26 +1144,6 @@ void FSteamVRInputDevice::InitControllerKeys()
 #pragma endregion
 
 #pragma region INDEX CONTROLLER
-	EKeys::AddKey(FKeyDetails(IndexControllerKeys::ValveIndex_Left_Thumbstick_Up, LOCTEXT("ValveIndex_Left_Thumbstick_Up", "Valve Index (L) Thumbstick Up"), FKeyDetails::GamepadKey | FKeyDetails::NotBlueprintBindableKey, "ValveIndex"));
-	EKeys::AddKey(FKeyDetails(IndexControllerKeys::ValveIndex_Left_Thumbstick_Down, LOCTEXT("ValveIndex_Left_Thumbstick_Down", "Valve Index (L) Thumbstick Down"), FKeyDetails::GamepadKey | FKeyDetails::NotBlueprintBindableKey, "ValveIndex"));
-	EKeys::AddKey(FKeyDetails(IndexControllerKeys::ValveIndex_Left_Thumbstick_Left, LOCTEXT("ValveIndex_Left_Thumbstick_Left", "Valve Index (L) Thumbstick Left"), FKeyDetails::GamepadKey | FKeyDetails::NotBlueprintBindableKey, "ValveIndex"));
-	EKeys::AddKey(FKeyDetails(IndexControllerKeys::ValveIndex_Left_Thumbstick_Right, LOCTEXT("ValveIndex_Left_Thumbstick_Right", "Valve Index (L) Thumbstick Right"), FKeyDetails::GamepadKey | FKeyDetails::NotBlueprintBindableKey, "ValveIndex"));
-
-	EKeys::AddKey(FKeyDetails(IndexControllerKeys::ValveIndex_Right_Thumbstick_Up, LOCTEXT("ValveIndex_Right_Thumbstick_Up", "Valve Index (R) Thumbstick Up"), FKeyDetails::GamepadKey | FKeyDetails::NotBlueprintBindableKey, "ValveIndex"));
-	EKeys::AddKey(FKeyDetails(IndexControllerKeys::ValveIndex_Right_Thumbstick_Down, LOCTEXT("ValveIndex_Right_Thumbstick_Down", "Valve Index (R) Thumbstick Down"), FKeyDetails::GamepadKey | FKeyDetails::NotBlueprintBindableKey, "ValveIndex"));
-	EKeys::AddKey(FKeyDetails(IndexControllerKeys::ValveIndex_Right_Thumbstick_Left, LOCTEXT("ValveIndex_Right_Thumbstick_Left", "Valve Index (R) Thumbstick Left"), FKeyDetails::GamepadKey | FKeyDetails::NotBlueprintBindableKey, "ValveIndex"));
-	EKeys::AddKey(FKeyDetails(IndexControllerKeys::ValveIndex_Right_Thumbstick_Right, LOCTEXT("ValveIndex_Right_Thumbstick_Right", "Valve Index (R) Thumbstick Right"), FKeyDetails::GamepadKey | FKeyDetails::NotBlueprintBindableKey, "ValveIndex"));
-
-	EKeys::AddKey(FKeyDetails(IndexControllerKeys::ValveIndex_Left_Trackpad_Up, LOCTEXT("ValveIndex_Left_Trackpad_Up", "Valve Index (L) Trackpad Up"), FKeyDetails::GamepadKey | FKeyDetails::NotBlueprintBindableKey, "ValveIndex"));
-	EKeys::AddKey(FKeyDetails(IndexControllerKeys::ValveIndex_Left_Trackpad_Down, LOCTEXT("ValveIndex_Left_Trackpad_Down", "Valve Index (L) Trackpad Down"), FKeyDetails::GamepadKey | FKeyDetails::NotBlueprintBindableKey, "ValveIndex"));
-	EKeys::AddKey(FKeyDetails(IndexControllerKeys::ValveIndex_Left_Trackpad_Left, LOCTEXT("ValveIndex_Left_Trackpad_Left", "Valve Index (L) Trackpad Left"), FKeyDetails::GamepadKey | FKeyDetails::NotBlueprintBindableKey, "ValveIndex"));
-	EKeys::AddKey(FKeyDetails(IndexControllerKeys::ValveIndex_Left_Trackpad_Right, LOCTEXT("ValveIndex_Left_Trackpad_Right", "Valve Index (L) Trackpad Right"), FKeyDetails::GamepadKey | FKeyDetails::NotBlueprintBindableKey, "ValveIndex"));
-
-	EKeys::AddKey(FKeyDetails(IndexControllerKeys::ValveIndex_Right_Trackpad_Up, LOCTEXT("ValveIndex_Right_Trackpad_Up", "Valve Index (R) Trackpad Up"), FKeyDetails::GamepadKey | FKeyDetails::NotBlueprintBindableKey, "ValveIndex"));
-	EKeys::AddKey(FKeyDetails(IndexControllerKeys::ValveIndex_Right_Trackpad_Down, LOCTEXT("ValveIndex_Right_Trackpad_Down", "Valve Index (R) Trackpad Down"), FKeyDetails::GamepadKey | FKeyDetails::NotBlueprintBindableKey, "ValveIndex"));
-	EKeys::AddKey(FKeyDetails(IndexControllerKeys::ValveIndex_Right_Trackpad_Left, LOCTEXT("ValveIndex_Right_Trackpad_Left", "Valve Index (R) Trackpad Left"), FKeyDetails::GamepadKey | FKeyDetails::NotBlueprintBindableKey, "ValveIndex"));
-	EKeys::AddKey(FKeyDetails(IndexControllerKeys::ValveIndex_Right_Trackpad_Right, LOCTEXT("ValveIndex_Right_Trackpad_Right", "Valve Index (R) Trackpad Right"), FKeyDetails::GamepadKey | FKeyDetails::NotBlueprintBindableKey, "ValveIndex"));
-
 	EKeys::AddKey(FKeyDetails(IndexControllerKeys::ValveIndex_Left_Pinch_Grab, LOCTEXT("ValveIndex_Left_Pinch_Grab", "Valve Index (L) Pinch Grab"), FKeyDetails::GamepadKey | FKeyDetails::NotBlueprintBindableKey, "ValveIndex"));
 	EKeys::AddKey(FKeyDetails(IndexControllerKeys::ValveIndex_Right_Pinch_Grab, LOCTEXT("ValveIndex_Right_Pinch_Grab", "Valve Index (R) Pinch Grab"), FKeyDetails::GamepadKey | FKeyDetails::NotBlueprintBindableKey, "ValveIndex"));
 
@@ -1166,42 +1151,6 @@ void FSteamVRInputDevice::InitControllerKeys()
 	EKeys::AddKey(FKeyDetails(IndexControllerKeys::ValveIndex_Right_Grip_Grab, LOCTEXT("ValveIndex_Right_Grip_Grab", "Valve Index (R) Grip Grab"), FKeyDetails::GamepadKey | FKeyDetails::NotBlueprintBindableKey, "ValveIndex"));
 #pragma endregion
 
-#pragma region OCULUS TOUCH CONTROLLER
-	EKeys::AddKey(FKeyDetails(OculusTouchKeys::OculusTouch_Left_Thumbstick_Up, LOCTEXT("OculusTouch_Left_Thumbstick_Up", "Oculus Touch (L) Thumbstick Up"), FKeyDetails::GamepadKey | FKeyDetails::NotBlueprintBindableKey, "OculusTouch"));
-	EKeys::AddKey(FKeyDetails(OculusTouchKeys::OculusTouch_Left_Thumbstick_Down, LOCTEXT("OculusTouch_Left_Thumbstick_Down", "Oculus Touch (L) Thumbstick Down"), FKeyDetails::GamepadKey | FKeyDetails::NotBlueprintBindableKey, "OculusTouch"));
-	EKeys::AddKey(FKeyDetails(OculusTouchKeys::OculusTouch_Left_Thumbstick_Left, LOCTEXT("OculusTouch_Left_Thumbstick_Left", "Oculus Touch (L) Thumbstick Left"), FKeyDetails::GamepadKey | FKeyDetails::NotBlueprintBindableKey, "OculusTouch"));
-	EKeys::AddKey(FKeyDetails(OculusTouchKeys::OculusTouch_Left_Thumbstick_Right, LOCTEXT("OculusTouch_Left_Thumbstick_Right", "Oculus Touch (L) Thumbstick Right"), FKeyDetails::GamepadKey | FKeyDetails::NotBlueprintBindableKey, "OculusTouch"));
-
-	EKeys::AddKey(FKeyDetails(OculusTouchKeys::OculusTouch_Right_Thumbstick_Up, LOCTEXT("OculusTouch_Right_Thumbstick_Up", "Oculus Touch (R) Thumbstick Up"), FKeyDetails::GamepadKey | FKeyDetails::NotBlueprintBindableKey, "OculusTouch"));
-	EKeys::AddKey(FKeyDetails(OculusTouchKeys::OculusTouch_Right_Thumbstick_Down, LOCTEXT("OculusTouch_Right_Thumbstick_Down", "Oculus Touch (R) Thumbstick Down"), FKeyDetails::GamepadKey | FKeyDetails::NotBlueprintBindableKey, "OculusTouch"));
-	EKeys::AddKey(FKeyDetails(OculusTouchKeys::OculusTouch_Right_Thumbstick_Left, LOCTEXT("OculusTouch_Right_Thumbstick_Left", "Oculus Touch (R) Thumbstick Left"), FKeyDetails::GamepadKey | FKeyDetails::NotBlueprintBindableKey, "OculusTouch"));
-	EKeys::AddKey(FKeyDetails(OculusTouchKeys::OculusTouch_Right_Thumbstick_Right, LOCTEXT("OculusTouch_Right_Thumbstick_Right", "Oculus Touch (R) Thumbstick Right"), FKeyDetails::GamepadKey | FKeyDetails::NotBlueprintBindableKey, "OculusTouch"));
-#pragma endregion
-
-#pragma region VIVE CONTROLLER
-	EKeys::AddKey(FKeyDetails(ViveControllerKeys::Vive_Left_Trackpad_Up, LOCTEXT("Vive_Left_Trackpad_Up", "Vive (L) Trackpad Up"), FKeyDetails::GamepadKey | FKeyDetails::NotBlueprintBindableKey, "Vive"));
-	EKeys::AddKey(FKeyDetails(ViveControllerKeys::Vive_Left_Trackpad_Down, LOCTEXT("Vive_Left_Trackpad_Down", "Vive (L) Trackpad Down"), FKeyDetails::GamepadKey | FKeyDetails::NotBlueprintBindableKey, "Vive"));
-	EKeys::AddKey(FKeyDetails(ViveControllerKeys::Vive_Left_Trackpad_Left, LOCTEXT("Vive_Left_Trackpad_Left", "Vive (L) Trackpad Left"), FKeyDetails::GamepadKey | FKeyDetails::NotBlueprintBindableKey, "Vive"));
-	EKeys::AddKey(FKeyDetails(ViveControllerKeys::Vive_Left_Trackpad_Right, LOCTEXT("Vive_Left_Trackpad_Right", "Vive (L) Trackpad Right"), FKeyDetails::GamepadKey | FKeyDetails::NotBlueprintBindableKey, "Vive"));
-
-	EKeys::AddKey(FKeyDetails(ViveControllerKeys::Vive_Right_Trackpad_Up, LOCTEXT("Vive_Right_Trackpad_Up", "Vive (R) Trackpad Up"), FKeyDetails::GamepadKey | FKeyDetails::NotBlueprintBindableKey, "Vive"));
-	EKeys::AddKey(FKeyDetails(ViveControllerKeys::Vive_Right_Trackpad_Down, LOCTEXT("Vive_Right_Trackpad_Down", "Vive (R) Trackpad Down"), FKeyDetails::GamepadKey | FKeyDetails::NotBlueprintBindableKey, "Vive"));
-	EKeys::AddKey(FKeyDetails(ViveControllerKeys::Vive_Right_Trackpad_Left, LOCTEXT("Vive_Right_Trackpad_Left", "Vive (R) Trackpad Left"), FKeyDetails::GamepadKey | FKeyDetails::NotBlueprintBindableKey, "Vive"));
-	EKeys::AddKey(FKeyDetails(ViveControllerKeys::Vive_Right_Trackpad_Right, LOCTEXT("Vive_Right_Trackpad_Right", "Vive (R) Trackpad Right"), FKeyDetails::GamepadKey | FKeyDetails::NotBlueprintBindableKey, "Vive"));
-#pragma endregion
-
-#pragma region WINDOWS MIXEDREALITY
-	EKeys::AddKey(FKeyDetails(WindowsMRKeys::MixedReality_Left_Trackpad_Up, LOCTEXT("MixedReality_Left_Trackpad_Up", "Mixed Reality (L) Trackpad Up"), FKeyDetails::GamepadKey | FKeyDetails::NotBlueprintBindableKey, "MixedReality"));
-	EKeys::AddKey(FKeyDetails(WindowsMRKeys::MixedReality_Left_Trackpad_Down, LOCTEXT("MixedReality_Left_Trackpad_Down", "Mixed Reality (L) Trackpad Down"), FKeyDetails::GamepadKey | FKeyDetails::NotBlueprintBindableKey, "MixedReality"));
-	EKeys::AddKey(FKeyDetails(WindowsMRKeys::MixedReality_Left_Trackpad_Left, LOCTEXT("MixedReality_Left_Trackpad_Left", "Mixed Reality (L) Trackpad Left"), FKeyDetails::GamepadKey | FKeyDetails::NotBlueprintBindableKey, "MixedReality"));
-	EKeys::AddKey(FKeyDetails(WindowsMRKeys::MixedReality_Left_Trackpad_Right, LOCTEXT("MixedReality_Left_Trackpad_Right", "Mixed Reality (L) Trackpad Right"), FKeyDetails::GamepadKey | FKeyDetails::NotBlueprintBindableKey, "MixedReality"));
-
-	EKeys::AddKey(FKeyDetails(WindowsMRKeys::MixedReality_Right_Trackpad_Up, LOCTEXT("MixedReality_Right_Trackpad_Up", "Mixed Reality (R) Trackpad Up"), FKeyDetails::GamepadKey | FKeyDetails::NotBlueprintBindableKey, "MixedReality"));
-	EKeys::AddKey(FKeyDetails(WindowsMRKeys::MixedReality_Right_Trackpad_Down, LOCTEXT("MixedReality_Right_Trackpad_Down", "Mixed Reality (R) Trackpad Down"), FKeyDetails::GamepadKey | FKeyDetails::NotBlueprintBindableKey, "MixedReality"));
-	EKeys::AddKey(FKeyDetails(WindowsMRKeys::MixedReality_Right_Trackpad_Left, LOCTEXT("MixedReality_Right_Trackpad_Left", "Mixed Reality (R) Trackpad Left"), FKeyDetails::GamepadKey | FKeyDetails::NotBlueprintBindableKey, "MixedReality"));
-	EKeys::AddKey(FKeyDetails(WindowsMRKeys::MixedReality_Right_Trackpad_Right, LOCTEXT("MixedReality_Right_Trackpad_Right", "Mixed Reality (R) Trackpad Right"), FKeyDetails::GamepadKey | FKeyDetails::NotBlueprintBindableKey, "MixedReality"));
-#pragma endregion
-
 #pragma region COSMOS KEYS
 	EKeys::AddMenuCategoryDisplayInfo("Cosmos", LOCTEXT("CosmosSubCategory", "HTC Cosmos"), TEXT("GraphEditor.PadEvent_16x"));
 
@@ -1261,6 +1210,15 @@ void FSteamVRInputDevice::RegenerateControllerBindings()
 	this->GenerateActionManifest(false, true, true, true);
 }
 
+void FSteamVRInputDevice::OnActionMappingsChanged()
+{
+	const USteamVRInputSettings* Settings = GetDefault<USteamVRInputSettings>();
+	if (!Settings->bDeveloperMode)
+	{
+		this->GenerateActionManifest(true, true, true, true);
+	}
+}
+
 bool FSteamVRInputDevice::GenerateAppManifest(FString ManifestPath, FString ProjectName, FString& OutAppKey, FString& OutAppManifestPath)
 {
 	// Set SteamVR AppKey
@@ -1585,7 +1543,6 @@ void FSteamVRInputDevice::GenerateActionBindings(TArray<FInputMapping> &InInputM
 	// Process Key Input Mappings
 	for (FSteamVRInputKeyMapping SteamVRKeyInputMapping : SteamVRKeyInputMappings)
 	{
-
 		// Check if this is a generic UE motion controller key
 		bool bHasSteamVRInputs = false;
 		if (bIsGenericController)
@@ -1663,7 +1620,7 @@ void FSteamVRInputDevice::GenerateActionBindings(TArray<FInputMapping> &InInputM
 					CurrentInputKeyName.Contains(TEXT("_A_"));
 				InputState.bIsFaceButton2 = CurrentInputKeyName.Contains(TEXT("FaceButton2"), ESearchCase::CaseSensitive, ESearchDir::FromEnd) ||
 					CurrentInputKeyName.Contains(TEXT("_B_"));
-				InputState.bIsAppMenu = CurrentInputKeyName.Contains(TEXT("_System_"));
+				InputState.bIsAppMenu = CurrentInputKeyName.Contains(TEXT("_Menu_"));
 				InputState.bIsProximity = CurrentInputKeyName.Contains(TEXT("_HMD_Proximity"));
 				
 				// Only handle proximity sensor for headsets
@@ -1945,7 +1902,7 @@ void FSteamVRInputDevice::GenerateActionBindings(TArray<FInputMapping> &InInputM
 				{
 					CachePath = FString(TEXT(ACTION_PATH_HEAD_PROXIMITY));
 					CacheMode = FName(TEXT("button"));
-					CacheType = FString(TEXT("click"));;
+					CacheType = FString(TEXT("click"));
 				}
 
 				if (!CacheType.IsEmpty())
@@ -2512,8 +2469,8 @@ void FSteamVRInputDevice::GenerateActionManifest(bool GenerateActions, bool Gene
 
 				// Setup the action fields
 				TArray<FString> ActionFields = {
-												 TEXT("name"), Action.Path,
-												 TEXT("type"), Action.GetActionTypeName(),
+					TEXT("name"), Action.Path,
+					TEXT("type"), Action.GetActionTypeName(),
 				};
 
 				// Add hand if skeleton
@@ -2600,8 +2557,8 @@ void FSteamVRInputDevice::GenerateActionManifest(bool GenerateActions, bool Gene
 
 	// Create action set objects
 	TArray<FString> StringFields = {
-									 "name", TEXT(ACTION_SET),
-									 "usage", TEXT("leftright")
+		"name", TEXT(ACTION_SET),
+		"usage", TEXT("leftright")
 	};
 
 	BuildJsonObject(StringFields, ActionSetObject);
@@ -2616,14 +2573,6 @@ void FSteamVRInputDevice::GenerateActionManifest(bool GenerateActions, bool Gene
 #pragma endregion
 
 #pragma region DEFAULT CONTROLLER BINDINGST
-#if WITH_EDITOR
-	// If we're running in the editor, build the controller bindings if they don't exist yet
-	if (GenerateBindings)
-	{
-		GenerateControllerBindings(ControllerBindingsPath, ControllerTypes, ControllerBindings, Actions, InputMappings, DeleteIfExists);
-	}
-#endif
-
 	// Start search for controller bindings files
 	TArray<FString> ControllerBindingFiles;
 	FileManager.FindFiles(ControllerBindingFiles, *ControllerBindingsPath, TEXT("*.json"));
@@ -2680,6 +2629,14 @@ void FSteamVRInputDevice::GenerateActionManifest(bool GenerateActions, bool Gene
 		}
 	}
 
+#if WITH_EDITOR
+	// If we're running in the editor, build the controller bindings if they don't exist yet
+	if (GenerateBindings)
+	{
+		GenerateControllerBindings(ControllerBindingsPath, ControllerTypes, ControllerBindings, Actions, InputMappings, DeleteIfExists);
+	}
+#endif
+
 	// Add the default bindings object to the action manifest
 	if (ControllerBindings.Num() == 0)
 	{
@@ -3291,30 +3248,30 @@ void FSteamVRInputDevice::SanitizeActions()
 
 void FSteamVRInputDevice::RegisterApplication(FString ManifestPath)
 {
-
 	if (VRSystem() && VRInput())
 	{
 		// Get Project Name this plugin is used in
 		uint32 AppProcessId = FPlatformProcess::GetCurrentProcessId();
 		GameFileName = FPaths::GetCleanFilename(FPlatformProcess::GetApplicationName(AppProcessId));
-		GameProjectName = GameFileName;
+		FString ProjectName;
 		if (GConfig)
 		{
 			GConfig->GetString(
 				TEXT("/Script/EngineSettings.GeneralProjectSettings"),
 				TEXT("ProjectName"),
-				GameProjectName,
+				ProjectName,
 				GGameIni
 			);
-
 		}
 
 		// Check for empty project name
-		if (GameProjectName.IsEmpty())
+		if (ProjectName.IsEmpty())
 		{
-			GameProjectName = GameFileName;
+			ProjectName = FApp::GetProjectName();
 		}
 
+		GameProjectName = FString::Printf(TEXT("%s-%u"), *ProjectName, FEngineVersion::Current().GetChangelist());
+
 #if WITH_EDITOR
 		if (VRApplications())
 		{
diff --git a/Engine/Plugins/Runtime/Steam/SteamVR/Source/SteamVRInputDevice/Private/SteamVRInputDeviceModule.cpp b/Engine/Plugins/Runtime/Steam/SteamVR/Source/SteamVRInputDevice/Private/SteamVRInputDeviceModule.cpp
--- a/Engine/Plugins/Runtime/Steam/SteamVR/Source/SteamVRInputDevice/Private/SteamVRInputDeviceModule.cpp
+++ b/Engine/Plugins/Runtime/Steam/SteamVR/Source/SteamVRInputDevice/Private/SteamVRInputDeviceModule.cpp
@@ -7,6 +7,11 @@
 #include "Engine/Engine.h"
 #include "IXRTrackingSystem.h"
 
+#if WITH_EDITOR
+#include "Editor/EditorEngine.h"
+#include "Editor.h"
+#endif
+
 class FSteamVRInputDeviceModule : public ISteamVRInputDeviceModule
 {
 	/* Creates a new instance of SteamVR Input Controller **/
@@ -15,7 +20,11 @@ class FSteamVRInputDeviceModule : public ISteamVRInputDeviceModule
 		static FName SystemName(TEXT("SteamVR"));
 		if (GEngine->XRSystem.IsValid() && (GEngine->XRSystem->GetSystemName() == SystemName))
 		{
-			return TSharedPtr<class IInputDevice>(new FSteamVRInputDevice(InMessageHandler));
+			TSharedPtr<class FSteamVRInputDevice> Device(new FSteamVRInputDevice(InMessageHandler));
+#if WITH_EDITOR
+			FEditorDelegates::OnActionAxisMappingsChanged.AddSP(Device.ToSharedRef(), &FSteamVRInputDevice::OnActionMappingsChanged);
+#endif
+			return Device;
 		}
 		return nullptr;
 	}
diff --git a/Engine/Plugins/Runtime/Steam/SteamVR/Source/SteamVRInputDevice/Public/SteamVRControllerKeys.h b/Engine/Plugins/Runtime/Steam/SteamVR/Source/SteamVRInputDevice/Public/SteamVRControllerKeys.h
--- a/Engine/Plugins/Runtime/Steam/SteamVR/Source/SteamVRInputDevice/Public/SteamVRControllerKeys.h
+++ b/Engine/Plugins/Runtime/Steam/SteamVR/Source/SteamVRInputDevice/Public/SteamVRControllerKeys.h
@@ -11,26 +11,6 @@ namespace GenericKeys
 
 namespace IndexControllerKeys
 {
-	// Thumbstick Directions
-	const FKey ValveIndex_Left_Thumbstick_Up("ValveIndex_Left_Thumbstick_Up");
-	const FKey ValveIndex_Left_Thumbstick_Down("ValveIndex_Left_Thumbstick_Down");
-	const FKey ValveIndex_Left_Thumbstick_Left("ValveIndex_Left_Thumbstick_Left");
-	const FKey ValveIndex_Left_Thumbstick_Right("ValveIndex_Left_Thumbstick_Right");
-	const FKey ValveIndex_Right_Thumbstick_Up("ValveIndex_Right_Thumbstick_Up");
-	const FKey ValveIndex_Right_Thumbstick_Down("ValveIndex_Right_Thumbstick_Down");
-	const FKey ValveIndex_Right_Thumbstick_Left("ValveIndex_Right_Thumbstick_Left");
-	const FKey ValveIndex_Right_Thumbstick_Right("ValveIndex_Right_Thumbstick_Right");
-
-	// Trackpad Directions
-	const FKey ValveIndex_Left_Trackpad_Up("ValveIndex_Left_Trackpad_Up");
-	const FKey ValveIndex_Left_Trackpad_Down("ValveIndex_Left_Trackpad_Down");
-	const FKey ValveIndex_Left_Trackpad_Left("ValveIndex_Left_Trackpad_Left");
-	const FKey ValveIndex_Left_Trackpad_Right("ValveIndex_Left_Trackpad_Right");
-	const FKey ValveIndex_Right_Trackpad_Up("ValveIndex_Right_Trackpad_Up");
-	const FKey ValveIndex_Right_Trackpad_Down("ValveIndex_Right_Trackpad_Down");
-	const FKey ValveIndex_Right_Trackpad_Left("ValveIndex_Right_Trackpad_Left");
-	const FKey ValveIndex_Right_Trackpad_Right("ValveIndex_Right_Trackpad_Right");
-
 	// Special
 	const FKey ValveIndex_Left_Grip_Grab("ValveIndex_Left_Grip_Grab");
 	const FKey ValveIndex_Left_Pinch_Grab("ValveIndex_Left_Pinch_Grab");
@@ -38,19 +18,6 @@ namespace IndexControllerKeys
 	const FKey ValveIndex_Right_Pinch_Grab("ValveIndex_Right_Pinch_Grab");
 }
 
-namespace ViveControllerKeys
-{
-	// Trackpad Directions
-	const FKey Vive_Left_Trackpad_Up("Vive_Left_Trackpad_Up");
-	const FKey Vive_Left_Trackpad_Down("Vive_Left_Trackpad_Down");
-	const FKey Vive_Left_Trackpad_Left("Vive_Left_Trackpad_Left");
-	const FKey Vive_Left_Trackpad_Right("Vive_Left_Trackpad_Right");
-	const FKey Vive_Right_Trackpad_Up("Vive_Right_Trackpad_Up");
-	const FKey Vive_Right_Trackpad_Down("Vive_Right_Trackpad_Down");
-	const FKey Vive_Right_Trackpad_Left("Vive_Right_Trackpad_Left");
-	const FKey Vive_Right_Trackpad_Right("Vive_Right_Trackpad_Right");
-}
-
 namespace CosmosKeys
 {
 	// Regular Controller Keys
@@ -98,28 +65,3 @@ namespace CosmosKeys
 	const FKey Cosmos_Right_Thumbstick_Right("Cosmos_Right_Thumbstick_Right");
 }
 
-namespace OculusTouchKeys
-{
-	// Thumbstick Directions
-	const FKey OculusTouch_Left_Thumbstick_Up("OculusTouch_Left_Thumbstick_Up");
-	const FKey OculusTouch_Left_Thumbstick_Down("OculusTouch_Left_Thumbstick_Down");
-	const FKey OculusTouch_Left_Thumbstick_Left("OculusTouch_Left_Thumbstick_Left");
-	const FKey OculusTouch_Left_Thumbstick_Right("OculusTouch_Left_Thumbstick_Right");
-	const FKey OculusTouch_Right_Thumbstick_Up("OculusTouch_Right_Thumbstick_Up");
-	const FKey OculusTouch_Right_Thumbstick_Down("OculusTouch_Right_Thumbstick_Down");
-	const FKey OculusTouch_Right_Thumbstick_Left("OculusTouch_Right_Thumbstick_Left");
-	const FKey OculusTouch_Right_Thumbstick_Right("OculusTouch_Right_Thumbstick_Right");
-}
-
-namespace WindowsMRKeys
-{
-	// Trackpad Directions
-	const FKey MixedReality_Left_Trackpad_Up("MixedReality_Left_Trackpad_Up");
-	const FKey MixedReality_Left_Trackpad_Down("MixedReality_Left_Trackpad_Down");
-	const FKey MixedReality_Left_Trackpad_Left("MixedReality_Left_Trackpad_Left");
-	const FKey MixedReality_Left_Trackpad_Right("MixedReality_Left_Trackpad_Right");
-	const FKey MixedReality_Right_Trackpad_Up("MixedReality_Right_Trackpad_Up");
-	const FKey MixedReality_Right_Trackpad_Down("MixedReality_Right_Trackpad_Down");
-	const FKey MixedReality_Right_Trackpad_Left("MixedReality_Right_Trackpad_Left");
-	const FKey MixedReality_Right_Trackpad_Right("MixedReality_Right_Trackpad_Right");
-}
diff --git a/Engine/Plugins/Runtime/Steam/SteamVR/Source/SteamVRInputDevice/Public/SteamVRInputDevice.h b/Engine/Plugins/Runtime/Steam/SteamVR/Source/SteamVRInputDevice/Public/SteamVRInputDevice.h
--- a/Engine/Plugins/Runtime/Steam/SteamVR/Source/SteamVRInputDevice/Public/SteamVRInputDevice.h
+++ b/Engine/Plugins/Runtime/Steam/SteamVR/Source/SteamVRInputDevice/Public/SteamVRInputDevice.h
@@ -95,9 +95,12 @@ public:
 	/** Have the action manifest regenerated. Used by the plugin Editor UI */
 	void RegenerateActionManifest();
 
-	/** Have the controller bindings regenerated. Used by the plugin Editor UI  */
+	/** Have the controller bindings regenerated. Used by the plugin Editor UI */
 	void RegenerateControllerBindings();
 
+	/** Have the manifest and bindings bindings regenerated. Used by a UI callback */
+	void OnActionMappingsChanged();
+
 	/** Reload the action manifest to the SteamVR system. Use if you changed the action manifest manually */
 	void ReloadActionManifest();
 #endif
diff --git a/Engine/Plugins/Runtime/Steam/SteamVR/Source/SteamVRInputDevice/SteamVRInputDevice.Build.cs b/Engine/Plugins/Runtime/Steam/SteamVR/Source/SteamVRInputDevice/SteamVRInputDevice.Build.cs
--- a/Engine/Plugins/Runtime/Steam/SteamVR/Source/SteamVRInputDevice/SteamVRInputDevice.Build.cs
+++ b/Engine/Plugins/Runtime/Steam/SteamVR/Source/SteamVRInputDevice/SteamVRInputDevice.Build.cs
@@ -45,7 +45,7 @@ public class SteamVRInputDevice : ModuleRules
 
         if (Target.bBuildEditor == true)
         {
-            DynamicallyLoadedModuleNames.Add("UnrealEd");
+            PrivateDependencyModuleNames.Add("UnrealEd");
         }
 
         AddEngineThirdPartyPrivateStaticDependencies(Target, "OpenVR");
diff --git a/Engine/Plugins/Runtime/WindowsMixedReality/Source/WindowsMixedRealitySpatialInput/Private/WindowsMixedRealitySpatialInput.cpp b/Engine/Plugins/Runtime/WindowsMixedReality/Source/WindowsMixedRealitySpatialInput/Private/WindowsMixedRealitySpatialInput.cpp
--- a/Engine/Plugins/Runtime/WindowsMixedReality/Source/WindowsMixedRealitySpatialInput/Private/WindowsMixedRealitySpatialInput.cpp
+++ b/Engine/Plugins/Runtime/WindowsMixedReality/Source/WindowsMixedRealitySpatialInput/Private/WindowsMixedRealitySpatialInput.cpp
@@ -112,6 +112,76 @@ namespace WindowsMixedReality
 			static_cast<float>(axisPosition));
 	}
 
+	FKey GetThumbstickDirectionKey(FVector2D TouchDir, HMDHand Hand)
+	{
+		const FVector2D UpDir(0.f, 1.f);
+		const FVector2D RightDir(1.f, 0.f);
+
+		const float VerticalDot = TouchDir | UpDir;
+		const float RightDot = TouchDir | RightDir;
+
+		if (VerticalDot >= UE_HALF_SQRT_2)
+		{
+			return (Hand == HMDHand::Left) ?
+				EKeys::MixedReality_Left_Thumbstick_Up :
+				EKeys::MixedReality_Right_Thumbstick_Up;
+		}
+		else if (VerticalDot <= -UE_HALF_SQRT_2)
+		{
+			return (Hand == HMDHand::Left) ?
+				EKeys::MixedReality_Left_Thumbstick_Down :
+				EKeys::MixedReality_Right_Thumbstick_Down;
+		}
+		else if (RightDot <= -UE_HALF_SQRT_2)
+		{
+			return (Hand == HMDHand::Left) ?
+				EKeys::MixedReality_Left_Thumbstick_Left :
+				EKeys::MixedReality_Right_Thumbstick_Left;
+		}
+		else if (RightDot >= UE_HALF_SQRT_2)
+		{
+			return (Hand == HMDHand::Left) ?
+				EKeys::MixedReality_Left_Thumbstick_Right :
+				EKeys::MixedReality_Right_Thumbstick_Right;
+		}
+		return EKeys::Invalid;
+	}
+
+	FKey GetTouchpadDirectionKey(FVector2D TouchDir, HMDHand Hand)
+	{
+		const FVector2D UpDir(0.f, 1.f);
+		const FVector2D RightDir(1.f, 0.f);
+
+		const float VerticalDot = TouchDir | UpDir;
+		const float RightDot = TouchDir | RightDir;
+
+		if (VerticalDot >= UE_HALF_SQRT_2)
+		{
+			return (Hand == HMDHand::Left) ?
+				EKeys::MixedReality_Left_Trackpad_Up :
+				EKeys::MixedReality_Right_Trackpad_Up;
+		}
+		else if (VerticalDot <= -UE_HALF_SQRT_2)
+		{
+			return (Hand == HMDHand::Left) ?
+				EKeys::MixedReality_Left_Trackpad_Down :
+				EKeys::MixedReality_Right_Trackpad_Down;
+		}
+		else if (RightDot <= -UE_HALF_SQRT_2)
+		{
+			return (Hand == HMDHand::Left) ?
+				EKeys::MixedReality_Left_Trackpad_Left :
+				EKeys::MixedReality_Right_Trackpad_Left;
+		}
+		else if (RightDot >= UE_HALF_SQRT_2)
+		{
+			return (Hand == HMDHand::Left) ?
+				EKeys::MixedReality_Left_Trackpad_Right :
+				EKeys::MixedReality_Right_Trackpad_Right;
+		}
+		return EKeys::Invalid;
+	}
+
 	// Gesture events come in from some microsoft thread, so we have to queue them up and send them from the game thread
 	// to avoid problems with systems that handle the events directly (UI is an example).
 	void FWindowsMixedRealitySpatialInput::EnqueueControllerButtonEvent(
@@ -169,7 +239,7 @@ namespace WindowsMixedReality
 	void FWindowsMixedRealitySpatialInput::SendAxisEvents(uint32 source)
 	{
 		FKey key;
-		FVector position;
+		FVector2D position;
 
 		for (int i = 0; i < 2; i++)
 		{
@@ -214,6 +284,25 @@ namespace WindowsMixedReality
 
 			SendControllerAxisEvent(MessageHandler, source, key, position.Y);
 
+			// Thumbstick direction
+			if (position.Size() > 0.5f)
+			{
+				key = GetThumbstickDirectionKey(position, hand);
+				if (!ThumbstickDirection[i].IsValid())
+				{
+					SendControllerButtonEvent(MessageHandler, source, key, HMDInputPressState::Pressed);
+				}
+				ThumbstickDirection[i] = key;
+			}
+			else
+			{
+				if (ThumbstickDirection[i].IsValid())
+				{
+					SendControllerButtonEvent(MessageHandler, source, ThumbstickDirection[i], HMDInputPressState::Released);
+				}
+				ThumbstickDirection[i] = EKeys::Invalid;
+			}
+
 			// Touchpad X
 			position.X = FWindowsMixedRealityStatics::GetAxisPosition(hand, HMDInputControllerAxes::TouchpadX);
 			key = (hand == HMDHand::Left) ?
@@ -241,6 +330,25 @@ namespace WindowsMixedReality
 			}
 
 			SendControllerAxisEvent(MessageHandler, source, key, position.Y);
+
+			// Touchpad direction
+			if (position.Size() > 0.5f)
+			{
+				key = GetTouchpadDirectionKey(position, hand);
+				if (!ThumbstickDirection[i].IsValid())
+				{
+					SendControllerButtonEvent(MessageHandler, source, key, HMDInputPressState::Pressed);
+				}
+				ThumbstickDirection[i] = key;
+			}
+			else
+			{
+				if (ThumbstickDirection[i].IsValid())
+				{
+					SendControllerButtonEvent(MessageHandler, source, ThumbstickDirection[i], HMDInputPressState::Released);
+				}
+				ThumbstickDirection[i] = EKeys::Invalid;
+			}
 		}
 	}
 
diff --git a/Engine/Plugins/Runtime/WindowsMixedReality/Source/WindowsMixedRealitySpatialInput/Private/WindowsMixedRealitySpatialInput.h b/Engine/Plugins/Runtime/WindowsMixedReality/Source/WindowsMixedRealitySpatialInput/Private/WindowsMixedRealitySpatialInput.h
--- a/Engine/Plugins/Runtime/WindowsMixedReality/Source/WindowsMixedRealitySpatialInput/Private/WindowsMixedRealitySpatialInput.h
+++ b/Engine/Plugins/Runtime/WindowsMixedReality/Source/WindowsMixedRealitySpatialInput/Private/WindowsMixedRealitySpatialInput.h
@@ -111,6 +111,8 @@ namespace WindowsMixedReality
 
 		bool isLeftTouchpadTouched = false;
 		bool isRightTouchpadTouched = false;
+		FKey ThumbstickDirection[2] = { EKeys::Invalid, EKeys::Invalid };
+		FKey TouchpadDirection[2] = { EKeys::Invalid, EKeys::Invalid };
 
 		FThreadSafeBool IsInitialized = false;
 
diff --git a/Engine/Source/Runtime/InputCore/Classes/InputCoreTypes.h b/Engine/Source/Runtime/InputCore/Classes/InputCoreTypes.h
--- a/Engine/Source/Runtime/InputCore/Classes/InputCoreTypes.h
+++ b/Engine/Source/Runtime/InputCore/Classes/InputCoreTypes.h
@@ -71,6 +71,7 @@ struct INPUTCORE_API FKey
 	bool IsBindableInBlueprints() const;
 	bool ShouldUpdateAxisWithoutSamples() const;
 	bool IsBindableToActions() const;
+	bool IsDeprecated() const;
 	FText GetDisplayName(bool bLongDisplayName = true) const;
 	FString ToString() const;
 	FName GetFName() const;
@@ -131,6 +132,7 @@ struct INPUTCORE_API FKeyDetails
 		VectorAxis				 = 1 << 6,
 		UpdateAxisWithoutSamples = 1 << 7,
 		NotActionBindableKey	 = 1 << 8,
+		Deprecated				 = 1 << 9,
 
 		NoFlags                 = 0,
 	};
@@ -147,6 +149,7 @@ struct INPUTCORE_API FKeyDetails
 	FORCEINLINE bool IsBindableInBlueprints() const { return bIsBindableInBlueprints != 0; }
 	FORCEINLINE bool ShouldUpdateAxisWithoutSamples() const { return bShouldUpdateAxisWithoutSamples != 0; }
 	FORCEINLINE bool IsBindableToActions() const { return bIsBindableToActions != 0; }
+	FORCEINLINE bool IsDeprecated() const { return bIsDeprecated != 0; }
 	FORCEINLINE FName GetMenuCategory() const { return MenuCategory; }
 	FText GetDisplayName(const bool bLongDisplayName = true) const;
 	FORCEINLINE const FKey& GetKey() const { return Key; }
@@ -173,6 +176,7 @@ private:
 	uint8 bIsBindableInBlueprints : 1;
 	uint8 bShouldUpdateAxisWithoutSamples : 1;
 	uint8 bIsBindableToActions : 1;
+	uint8 bIsDeprecated : 1;
 	EInputAxisType AxisType;
 
 	TAttribute<FText> LongDisplayName;
@@ -525,6 +529,10 @@ struct INPUTCORE_API EKeys
 	static const FKey Vive_Left_Trackpad_Y;
 	static const FKey Vive_Left_Trackpad_Click;
 	static const FKey Vive_Left_Trackpad_Touch;
+	static const FKey Vive_Left_Trackpad_Up;
+	static const FKey Vive_Left_Trackpad_Down;
+	static const FKey Vive_Left_Trackpad_Left;
+	static const FKey Vive_Left_Trackpad_Right;
 	static const FKey Vive_Right_System_Click;
 	static const FKey Vive_Right_Grip_Click;
 	static const FKey Vive_Right_Menu_Click;
@@ -534,6 +542,10 @@ struct INPUTCORE_API EKeys
 	static const FKey Vive_Right_Trackpad_Y;
 	static const FKey Vive_Right_Trackpad_Click;
 	static const FKey Vive_Right_Trackpad_Touch;
+	static const FKey Vive_Right_Trackpad_Up;
+	static const FKey Vive_Right_Trackpad_Down;
+	static const FKey Vive_Right_Trackpad_Left;
+	static const FKey Vive_Right_Trackpad_Right;
 
 	// Microsoft Mixed Reality Motion Controller
 	static const FKey MixedReality_Left_Menu_Click;
@@ -543,10 +555,18 @@ struct INPUTCORE_API EKeys
 	static const FKey MixedReality_Left_Thumbstick_X;
 	static const FKey MixedReality_Left_Thumbstick_Y;
 	static const FKey MixedReality_Left_Thumbstick_Click;
+	static const FKey MixedReality_Left_Thumbstick_Up;
+	static const FKey MixedReality_Left_Thumbstick_Down;
+	static const FKey MixedReality_Left_Thumbstick_Left;
+	static const FKey MixedReality_Left_Thumbstick_Right;
 	static const FKey MixedReality_Left_Trackpad_X;
 	static const FKey MixedReality_Left_Trackpad_Y;
 	static const FKey MixedReality_Left_Trackpad_Click;
 	static const FKey MixedReality_Left_Trackpad_Touch;
+	static const FKey MixedReality_Left_Trackpad_Up;
+	static const FKey MixedReality_Left_Trackpad_Down;
+	static const FKey MixedReality_Left_Trackpad_Left;
+	static const FKey MixedReality_Left_Trackpad_Right;
 	static const FKey MixedReality_Right_Menu_Click;
 	static const FKey MixedReality_Right_Grip_Click;
 	static const FKey MixedReality_Right_Trigger_Click;
@@ -554,10 +574,18 @@ struct INPUTCORE_API EKeys
 	static const FKey MixedReality_Right_Thumbstick_X;
 	static const FKey MixedReality_Right_Thumbstick_Y;
 	static const FKey MixedReality_Right_Thumbstick_Click;
+	static const FKey MixedReality_Right_Thumbstick_Up;
+	static const FKey MixedReality_Right_Thumbstick_Down;
+	static const FKey MixedReality_Right_Thumbstick_Left;
+	static const FKey MixedReality_Right_Thumbstick_Right;
 	static const FKey MixedReality_Right_Trackpad_X;
 	static const FKey MixedReality_Right_Trackpad_Y;
 	static const FKey MixedReality_Right_Trackpad_Click;
 	static const FKey MixedReality_Right_Trackpad_Touch;
+	static const FKey MixedReality_Right_Trackpad_Up;
+	static const FKey MixedReality_Right_Trackpad_Down;
+	static const FKey MixedReality_Right_Trackpad_Left;
+	static const FKey MixedReality_Right_Trackpad_Right;
 
 	// Oculus Go Controller
 	static const FKey OculusGo_Left_System_Click;
@@ -590,6 +618,10 @@ struct INPUTCORE_API EKeys
 	static const FKey OculusTouch_Left_Thumbstick_Y;
 	static const FKey OculusTouch_Left_Thumbstick_Click;
 	static const FKey OculusTouch_Left_Thumbstick_Touch;
+	static const FKey OculusTouch_Left_Thumbstick_Up;
+	static const FKey OculusTouch_Left_Thumbstick_Down;
+	static const FKey OculusTouch_Left_Thumbstick_Left;
+	static const FKey OculusTouch_Left_Thumbstick_Right;
 	static const FKey OculusTouch_Right_A_Click;
 	static const FKey OculusTouch_Right_B_Click;
 	static const FKey OculusTouch_Right_A_Touch;
@@ -604,6 +636,10 @@ struct INPUTCORE_API EKeys
 	static const FKey OculusTouch_Right_Thumbstick_Y;
 	static const FKey OculusTouch_Right_Thumbstick_Click;
 	static const FKey OculusTouch_Right_Thumbstick_Touch;
+	static const FKey OculusTouch_Right_Thumbstick_Up;
+	static const FKey OculusTouch_Right_Thumbstick_Down;
+	static const FKey OculusTouch_Right_Thumbstick_Left;
+	static const FKey OculusTouch_Right_Thumbstick_Right;
 
 	// Valve Index Controller
 	static const FKey ValveIndex_Left_A_Click;
@@ -622,11 +658,19 @@ struct INPUTCORE_API EKeys
 	static const FKey ValveIndex_Left_Thumbstick_Y;
 	static const FKey ValveIndex_Left_Thumbstick_Click;
 	static const FKey ValveIndex_Left_Thumbstick_Touch;
+	static const FKey ValveIndex_Left_Thumbstick_Up;
+	static const FKey ValveIndex_Left_Thumbstick_Down;
+	static const FKey ValveIndex_Left_Thumbstick_Left;
+	static const FKey ValveIndex_Left_Thumbstick_Right;
 	static const FKey ValveIndex_Left_Trackpad_X;
 	static const FKey ValveIndex_Left_Trackpad_Y;
 	static const FKey ValveIndex_Left_Trackpad_Click;
 	static const FKey ValveIndex_Left_Trackpad_Force;
 	static const FKey ValveIndex_Left_Trackpad_Touch;
+	static const FKey ValveIndex_Left_Trackpad_Up;
+	static const FKey ValveIndex_Left_Trackpad_Down;
+	static const FKey ValveIndex_Left_Trackpad_Left;
+	static const FKey ValveIndex_Left_Trackpad_Right;
 	static const FKey ValveIndex_Right_A_Click;
 	static const FKey ValveIndex_Right_B_Click;
 	static const FKey ValveIndex_Right_A_Touch;
@@ -643,11 +687,19 @@ struct INPUTCORE_API EKeys
 	static const FKey ValveIndex_Right_Thumbstick_Y;
 	static const FKey ValveIndex_Right_Thumbstick_Click;
 	static const FKey ValveIndex_Right_Thumbstick_Touch;
+	static const FKey ValveIndex_Right_Thumbstick_Up;
+	static const FKey ValveIndex_Right_Thumbstick_Down;
+	static const FKey ValveIndex_Right_Thumbstick_Left;
+	static const FKey ValveIndex_Right_Thumbstick_Right;
 	static const FKey ValveIndex_Right_Trackpad_X;
 	static const FKey ValveIndex_Right_Trackpad_Y;
 	static const FKey ValveIndex_Right_Trackpad_Click;
 	static const FKey ValveIndex_Right_Trackpad_Force;
 	static const FKey ValveIndex_Right_Trackpad_Touch;
+	static const FKey ValveIndex_Right_Trackpad_Up;
+	static const FKey ValveIndex_Right_Trackpad_Down;
+	static const FKey ValveIndex_Right_Trackpad_Left;
+	static const FKey ValveIndex_Right_Trackpad_Right;
 
 	// Virtual buttons that use other buttons depending on the platform
 	static const FKey Virtual_Accept;
diff --git a/Engine/Source/Runtime/InputCore/Private/InputCoreTypes.cpp b/Engine/Source/Runtime/InputCore/Private/InputCoreTypes.cpp
--- a/Engine/Source/Runtime/InputCore/Private/InputCoreTypes.cpp
+++ b/Engine/Source/Runtime/InputCore/Private/InputCoreTypes.cpp
@@ -336,6 +336,10 @@ const FKey EKeys::Vive_Left_Trackpad_X("Vive_Left_Trackpad_X");
 const FKey EKeys::Vive_Left_Trackpad_Y("Vive_Left_Trackpad_Y");
 const FKey EKeys::Vive_Left_Trackpad_Click("Vive_Left_Trackpad_Click");
 const FKey EKeys::Vive_Left_Trackpad_Touch("Vive_Left_Trackpad_Touch");
+const FKey EKeys::Vive_Left_Trackpad_Up("Vive_Left_Trackpad_Up");
+const FKey EKeys::Vive_Left_Trackpad_Down("Vive_Left_Trackpad_Down");
+const FKey EKeys::Vive_Left_Trackpad_Left("Vive_Left_Trackpad_Left");
+const FKey EKeys::Vive_Left_Trackpad_Right("Vive_Left_Trackpad_Right");
 const FKey EKeys::Vive_Right_System_Click("Vive_Right_System_Click");
 const FKey EKeys::Vive_Right_Grip_Click("Vive_Right_Grip_Click");
 const FKey EKeys::Vive_Right_Menu_Click("Vive_Right_Menu_Click");
@@ -345,6 +349,10 @@ const FKey EKeys::Vive_Right_Trackpad_X("Vive_Right_Trackpad_X");
 const FKey EKeys::Vive_Right_Trackpad_Y("Vive_Right_Trackpad_Y");
 const FKey EKeys::Vive_Right_Trackpad_Click("Vive_Right_Trackpad_Click");
 const FKey EKeys::Vive_Right_Trackpad_Touch("Vive_Right_Trackpad_Touch");
+const FKey EKeys::Vive_Right_Trackpad_Up("Vive_Right_Trackpad_Up");
+const FKey EKeys::Vive_Right_Trackpad_Down("Vive_Right_Trackpad_Down");
+const FKey EKeys::Vive_Right_Trackpad_Left("Vive_Right_Trackpad_Left");
+const FKey EKeys::Vive_Right_Trackpad_Right("Vive_Right_Trackpad_Right");
 
 // Microsoft Mixed Reality Motion Controller
 const FKey EKeys::MixedReality_Left_Menu_Click("MixedReality_Left_Menu_Click");
@@ -354,10 +362,18 @@ const FKey EKeys::MixedReality_Left_Trigger_Axis("MixedReality_Left_Trigger_Axis
 const FKey EKeys::MixedReality_Left_Thumbstick_X("MixedReality_Left_Thumbstick_X");
 const FKey EKeys::MixedReality_Left_Thumbstick_Y("MixedReality_Left_Thumbstick_Y");
 const FKey EKeys::MixedReality_Left_Thumbstick_Click("MixedReality_Left_Thumbstick_Click");
+const FKey EKeys::MixedReality_Left_Thumbstick_Up("MixedReality_Left_Thumbstick_Up");
+const FKey EKeys::MixedReality_Left_Thumbstick_Down("MixedReality_Left_Thumbstick_Down");
+const FKey EKeys::MixedReality_Left_Thumbstick_Left("MixedReality_Left_Thumbstick_Left");
+const FKey EKeys::MixedReality_Left_Thumbstick_Right("MixedReality_Left_Thumbstick_Right");
 const FKey EKeys::MixedReality_Left_Trackpad_X("MixedReality_Left_Trackpad_X");
 const FKey EKeys::MixedReality_Left_Trackpad_Y("MixedReality_Left_Trackpad_Y");
 const FKey EKeys::MixedReality_Left_Trackpad_Click("MixedReality_Left_Trackpad_Click");
 const FKey EKeys::MixedReality_Left_Trackpad_Touch("MixedReality_Left_Trackpad_Touch");
+const FKey EKeys::MixedReality_Left_Trackpad_Up("MixedReality_Left_Trackpad_Up");
+const FKey EKeys::MixedReality_Left_Trackpad_Down("MixedReality_Left_Trackpad_Down");
+const FKey EKeys::MixedReality_Left_Trackpad_Left("MixedReality_Left_Trackpad_Left");
+const FKey EKeys::MixedReality_Left_Trackpad_Right("MixedReality_Left_Trackpad_Right");
 const FKey EKeys::MixedReality_Right_Menu_Click("MixedReality_Right_Menu_Click");
 const FKey EKeys::MixedReality_Right_Grip_Click("MixedReality_Right_Grip_Click");
 const FKey EKeys::MixedReality_Right_Trigger_Click("MixedReality_Right_Trigger_Click");
@@ -365,10 +381,18 @@ const FKey EKeys::MixedReality_Right_Trigger_Axis("MixedReality_Right_Trigger_Ax
 const FKey EKeys::MixedReality_Right_Thumbstick_X("MixedReality_Right_Thumbstick_X");
 const FKey EKeys::MixedReality_Right_Thumbstick_Y("MixedReality_Right_Thumbstick_Y");
 const FKey EKeys::MixedReality_Right_Thumbstick_Click("MixedReality_Right_Thumbstick_Click");
+const FKey EKeys::MixedReality_Right_Thumbstick_Up("MixedReality_Right_Thumbstick_Up");
+const FKey EKeys::MixedReality_Right_Thumbstick_Down("MixedReality_Right_Thumbstick_Down");
+const FKey EKeys::MixedReality_Right_Thumbstick_Left("MixedReality_Right_Thumbstick_Left");
+const FKey EKeys::MixedReality_Right_Thumbstick_Right("MixedReality_Right_Thumbstick_Right");
 const FKey EKeys::MixedReality_Right_Trackpad_X("MixedReality_Right_Trackpad_X");
 const FKey EKeys::MixedReality_Right_Trackpad_Y("MixedReality_Right_Trackpad_Y");
 const FKey EKeys::MixedReality_Right_Trackpad_Click("MixedReality_Right_Trackpad_Click");
 const FKey EKeys::MixedReality_Right_Trackpad_Touch("MixedReality_Right_Trackpad_Touch");
+const FKey EKeys::MixedReality_Right_Trackpad_Up("MixedReality_Right_Trackpad_Up");
+const FKey EKeys::MixedReality_Right_Trackpad_Down("MixedReality_Right_Trackpad_Down");
+const FKey EKeys::MixedReality_Right_Trackpad_Left("MixedReality_Right_Trackpad_Left");
+const FKey EKeys::MixedReality_Right_Trackpad_Right("MixedReality_Right_Trackpad_Right");
 
 // Oculus Go Controller
 const FKey EKeys::OculusGo_Left_System_Click("OculusGo_Left_System_Click");
@@ -401,6 +425,10 @@ const FKey EKeys::OculusTouch_Left_Thumbstick_X("OculusTouch_Left_Thumbstick_X")
 const FKey EKeys::OculusTouch_Left_Thumbstick_Y("OculusTouch_Left_Thumbstick_Y");
 const FKey EKeys::OculusTouch_Left_Thumbstick_Click("OculusTouch_Left_Thumbstick_Click");
 const FKey EKeys::OculusTouch_Left_Thumbstick_Touch("OculusTouch_Left_Thumbstick_Touch");
+const FKey EKeys::OculusTouch_Left_Thumbstick_Up("OculusTouch_Left_Thumbstick_Up");
+const FKey EKeys::OculusTouch_Left_Thumbstick_Down("OculusTouch_Left_Thumbstick_Down");
+const FKey EKeys::OculusTouch_Left_Thumbstick_Left("OculusTouch_Left_Thumbstick_Left");
+const FKey EKeys::OculusTouch_Left_Thumbstick_Right("OculusTouch_Left_Thumbstick_Right");
 const FKey EKeys::OculusTouch_Right_A_Click("OculusTouch_Right_A_Click");
 const FKey EKeys::OculusTouch_Right_B_Click("OculusTouch_Right_B_Click");
 const FKey EKeys::OculusTouch_Right_A_Touch("OculusTouch_Right_A_Touch");
@@ -415,6 +443,10 @@ const FKey EKeys::OculusTouch_Right_Thumbstick_X("OculusTouch_Right_Thumbstick_X
 const FKey EKeys::OculusTouch_Right_Thumbstick_Y("OculusTouch_Right_Thumbstick_Y");
 const FKey EKeys::OculusTouch_Right_Thumbstick_Click("OculusTouch_Right_Thumbstick_Click");
 const FKey EKeys::OculusTouch_Right_Thumbstick_Touch("OculusTouch_Right_Thumbstick_Touch");
+const FKey EKeys::OculusTouch_Right_Thumbstick_Up("OculusTouch_Right_Thumbstick_Up");
+const FKey EKeys::OculusTouch_Right_Thumbstick_Down("OculusTouch_Right_Thumbstick_Down");
+const FKey EKeys::OculusTouch_Right_Thumbstick_Left("OculusTouch_Right_Thumbstick_Left");
+const FKey EKeys::OculusTouch_Right_Thumbstick_Right("OculusTouch_Right_Thumbstick_Right");
 
 // Valve Index Controller
 const FKey EKeys::ValveIndex_Left_A_Click("ValveIndex_Left_A_Click");
@@ -433,11 +465,19 @@ const FKey EKeys::ValveIndex_Left_Thumbstick_X("ValveIndex_Left_Thumbstick_X");
 const FKey EKeys::ValveIndex_Left_Thumbstick_Y("ValveIndex_Left_Thumbstick_Y");
 const FKey EKeys::ValveIndex_Left_Thumbstick_Click("ValveIndex_Left_Thumbstick_Click");
 const FKey EKeys::ValveIndex_Left_Thumbstick_Touch("ValveIndex_Left_Thumbstick_Touch");
+const FKey EKeys::ValveIndex_Left_Thumbstick_Up("ValveIndex_Left_Thumbstick_Up");
+const FKey EKeys::ValveIndex_Left_Thumbstick_Down("ValveIndex_Left_Thumbstick_Down");
+const FKey EKeys::ValveIndex_Left_Thumbstick_Left("ValveIndex_Left_Thumbstick_Left");
+const FKey EKeys::ValveIndex_Left_Thumbstick_Right("ValveIndex_Left_Thumbstick_Right");
 const FKey EKeys::ValveIndex_Left_Trackpad_X("ValveIndex_Left_Trackpad_X");
 const FKey EKeys::ValveIndex_Left_Trackpad_Y("ValveIndex_Left_Trackpad_Y");
 const FKey EKeys::ValveIndex_Left_Trackpad_Click("ValveIndex_Left_Trackpad_Click");
 const FKey EKeys::ValveIndex_Left_Trackpad_Force("ValveIndex_Left_Trackpad_Force");
 const FKey EKeys::ValveIndex_Left_Trackpad_Touch("ValveIndex_Left_Trackpad_Touch");
+const FKey EKeys::ValveIndex_Left_Trackpad_Up("ValveIndex_Left_Trackpad_Up");
+const FKey EKeys::ValveIndex_Left_Trackpad_Down("ValveIndex_Left_Trackpad_Down");
+const FKey EKeys::ValveIndex_Left_Trackpad_Left("ValveIndex_Left_Trackpad_Left");
+const FKey EKeys::ValveIndex_Left_Trackpad_Right("ValveIndex_Left_Trackpad_Right");
 const FKey EKeys::ValveIndex_Right_A_Click("ValveIndex_Right_A_Click");
 const FKey EKeys::ValveIndex_Right_B_Click("ValveIndex_Right_B_Click");
 const FKey EKeys::ValveIndex_Right_A_Touch("ValveIndex_Right_A_Touch");
@@ -454,11 +494,19 @@ const FKey EKeys::ValveIndex_Right_Thumbstick_X("ValveIndex_Right_Thumbstick_X")
 const FKey EKeys::ValveIndex_Right_Thumbstick_Y("ValveIndex_Right_Thumbstick_Y");
 const FKey EKeys::ValveIndex_Right_Thumbstick_Click("ValveIndex_Right_Thumbstick_Click");
 const FKey EKeys::ValveIndex_Right_Thumbstick_Touch("ValveIndex_Right_Thumbstick_Touch");
+const FKey EKeys::ValveIndex_Right_Thumbstick_Up("ValveIndex_Right_Thumbstick_Up");
+const FKey EKeys::ValveIndex_Right_Thumbstick_Down("ValveIndex_Right_Thumbstick_Down");
+const FKey EKeys::ValveIndex_Right_Thumbstick_Left("ValveIndex_Right_Thumbstick_Left");
+const FKey EKeys::ValveIndex_Right_Thumbstick_Right("ValveIndex_Right_Thumbstick_Right");
 const FKey EKeys::ValveIndex_Right_Trackpad_X("ValveIndex_Right_Trackpad_X");
 const FKey EKeys::ValveIndex_Right_Trackpad_Y("ValveIndex_Right_Trackpad_Y");
 const FKey EKeys::ValveIndex_Right_Trackpad_Click("ValveIndex_Right_Trackpad_Click");
 const FKey EKeys::ValveIndex_Right_Trackpad_Force("ValveIndex_Right_Trackpad_Force");
 const FKey EKeys::ValveIndex_Right_Trackpad_Touch("ValveIndex_Right_Trackpad_Touch");
+const FKey EKeys::ValveIndex_Right_Trackpad_Up("ValveIndex_Right_Trackpad_Up");
+const FKey EKeys::ValveIndex_Right_Trackpad_Down("ValveIndex_Right_Trackpad_Down");
+const FKey EKeys::ValveIndex_Right_Trackpad_Left("ValveIndex_Right_Trackpad_Left");
+const FKey EKeys::ValveIndex_Right_Trackpad_Right("ValveIndex_Right_Trackpad_Right");
 
 const FKey EKeys::Invalid(NAME_None);
 
@@ -497,9 +545,10 @@ void FKeyDetails::CommonInit(const uint16 InKeyFlags)
 	bIsGamepadKey = ((InKeyFlags & EKeyFlags::GamepadKey) != 0);
 	bIsTouch = ((InKeyFlags & EKeyFlags::Touch) != 0);
 	bIsMouseButton = ((InKeyFlags & EKeyFlags::MouseButton) != 0);
-	bIsBindableInBlueprints = ((~InKeyFlags & EKeyFlags::NotBlueprintBindableKey) != 0);
+	bIsBindableInBlueprints = ((~InKeyFlags & EKeyFlags::NotBlueprintBindableKey) != 0) && ((~InKeyFlags & EKeyFlags::Deprecated) != 0);
 	bShouldUpdateAxisWithoutSamples = ((InKeyFlags & EKeyFlags::UpdateAxisWithoutSamples) != 0);
-	bIsBindableToActions = ((~InKeyFlags & EKeyFlags::NotActionBindableKey) != 0);
+	bIsBindableToActions = ((~InKeyFlags & EKeyFlags::NotActionBindableKey) != 0) && ((~InKeyFlags & EKeyFlags::Deprecated) != 0);
+	bIsDeprecated = ((InKeyFlags & EKeyFlags::Deprecated) != 0);
 
 	if ((InKeyFlags & EKeyFlags::FloatAxis) != 0)
 	{
@@ -764,63 +813,63 @@ void EKeys::Initialize()
 
 	// Motion Controllers
 	//	Left Controller
-	AddKey(FKeyDetails(EKeys::MotionController_Left_FaceButton1, LOCTEXT("MotionController_Left_FaceButton1", "MotionController (L) FaceButton1"), FKeyDetails::GamepadKey | FKeyDetails::NotBlueprintBindableKey | FKeyDetails::NotActionBindableKey));
-	AddKey(FKeyDetails(EKeys::MotionController_Left_FaceButton2, LOCTEXT("MotionController_Left_FaceButton2", "MotionController (L) FaceButton2"), FKeyDetails::GamepadKey | FKeyDetails::NotBlueprintBindableKey | FKeyDetails::NotActionBindableKey));
-	AddKey(FKeyDetails(EKeys::MotionController_Left_FaceButton3, LOCTEXT("MotionController_Left_FaceButton3", "MotionController (L) FaceButton3"), FKeyDetails::GamepadKey | FKeyDetails::NotBlueprintBindableKey | FKeyDetails::NotActionBindableKey));
-	AddKey(FKeyDetails(EKeys::MotionController_Left_FaceButton4, LOCTEXT("MotionController_Left_FaceButton4", "MotionController (L) FaceButton4"), FKeyDetails::GamepadKey | FKeyDetails::NotBlueprintBindableKey | FKeyDetails::NotActionBindableKey));
-	AddKey(FKeyDetails(EKeys::MotionController_Left_FaceButton5, LOCTEXT("MotionController_Left_FaceButton5", "MotionController (L) FaceButton5"), FKeyDetails::GamepadKey | FKeyDetails::NotBlueprintBindableKey | FKeyDetails::NotActionBindableKey));
-	AddKey(FKeyDetails(EKeys::MotionController_Left_FaceButton6, LOCTEXT("MotionController_Left_FaceButton6", "MotionController (L) FaceButton6"), FKeyDetails::GamepadKey | FKeyDetails::NotBlueprintBindableKey | FKeyDetails::NotActionBindableKey));
-	AddKey(FKeyDetails(EKeys::MotionController_Left_FaceButton7, LOCTEXT("MotionController_Left_FaceButton7", "MotionController (L) FaceButton7"), FKeyDetails::GamepadKey | FKeyDetails::NotBlueprintBindableKey | FKeyDetails::NotActionBindableKey));
-	AddKey(FKeyDetails(EKeys::MotionController_Left_FaceButton8, LOCTEXT("MotionController_Left_FaceButton8", "MotionController (L) FaceButton8"), FKeyDetails::GamepadKey | FKeyDetails::NotBlueprintBindableKey | FKeyDetails::NotActionBindableKey));
-
-	AddKey(FKeyDetails(EKeys::MotionController_Left_Shoulder, LOCTEXT("MotionController_Left_Shoulder", "MotionController (L) Shoulder"), FKeyDetails::GamepadKey | FKeyDetails::NotBlueprintBindableKey | FKeyDetails::NotActionBindableKey));
-	AddKey(FKeyDetails(EKeys::MotionController_Left_Trigger, LOCTEXT("MotionController_Left_Trigger", "MotionController (L) Trigger"), FKeyDetails::GamepadKey | FKeyDetails::NotBlueprintBindableKey | FKeyDetails::NotActionBindableKey));
-
-	AddKey(FKeyDetails(EKeys::MotionController_Left_Grip1, LOCTEXT("MotionController_Left_Grip1", "MotionController (L) Grip1"), FKeyDetails::GamepadKey | FKeyDetails::NotBlueprintBindableKey | FKeyDetails::NotActionBindableKey));
-	AddKey(FKeyDetails(EKeys::MotionController_Left_Grip2, LOCTEXT("MotionController_Left_Grip2", "MotionController (L) Grip2"), FKeyDetails::GamepadKey | FKeyDetails::NotBlueprintBindableKey | FKeyDetails::NotActionBindableKey));
-
-	AddKey(FKeyDetails(EKeys::MotionController_Left_Thumbstick, LOCTEXT("MotionController_Left_Thumbstick", "MotionController (L) Thumbstick"), FKeyDetails::GamepadKey | FKeyDetails::NotBlueprintBindableKey | FKeyDetails::NotActionBindableKey));
-	AddKey(FKeyDetails(EKeys::MotionController_Left_Thumbstick_Up, LOCTEXT("MotionController_Left_Thumbstick_Up", "MotionController (L) Thumbstick Up"), FKeyDetails::GamepadKey | FKeyDetails::NotBlueprintBindableKey | FKeyDetails::NotActionBindableKey));
-	AddKey(FKeyDetails(EKeys::MotionController_Left_Thumbstick_Down, LOCTEXT("MotionController_Left_Thumbstick_Down", "MotionController (L) Thumbstick Down"), FKeyDetails::GamepadKey | FKeyDetails::NotBlueprintBindableKey | FKeyDetails::NotActionBindableKey));
-	AddKey(FKeyDetails(EKeys::MotionController_Left_Thumbstick_Left, LOCTEXT("MotionController_Left_Thumbstick_Left", "MotionController (L) Thumbstick Left"), FKeyDetails::GamepadKey | FKeyDetails::NotBlueprintBindableKey | FKeyDetails::NotActionBindableKey));
-	AddKey(FKeyDetails(EKeys::MotionController_Left_Thumbstick_Right, LOCTEXT("MotionController_Left_Thumbstick_Right", "MotionController (L) Thumbstick Right"), FKeyDetails::GamepadKey | FKeyDetails::NotBlueprintBindableKey | FKeyDetails::NotActionBindableKey));
+	AddKey(FKeyDetails(EKeys::MotionController_Left_FaceButton1, LOCTEXT("MotionController_Left_FaceButton1", "MotionController (L) FaceButton1"), FKeyDetails::GamepadKey | FKeyDetails::Deprecated));
+	AddKey(FKeyDetails(EKeys::MotionController_Left_FaceButton2, LOCTEXT("MotionController_Left_FaceButton2", "MotionController (L) FaceButton2"), FKeyDetails::GamepadKey | FKeyDetails::Deprecated));
+	AddKey(FKeyDetails(EKeys::MotionController_Left_FaceButton3, LOCTEXT("MotionController_Left_FaceButton3", "MotionController (L) FaceButton3"), FKeyDetails::GamepadKey | FKeyDetails::Deprecated));
+	AddKey(FKeyDetails(EKeys::MotionController_Left_FaceButton4, LOCTEXT("MotionController_Left_FaceButton4", "MotionController (L) FaceButton4"), FKeyDetails::GamepadKey | FKeyDetails::Deprecated));
+	AddKey(FKeyDetails(EKeys::MotionController_Left_FaceButton5, LOCTEXT("MotionController_Left_FaceButton5", "MotionController (L) FaceButton5"), FKeyDetails::GamepadKey | FKeyDetails::Deprecated));
+	AddKey(FKeyDetails(EKeys::MotionController_Left_FaceButton6, LOCTEXT("MotionController_Left_FaceButton6", "MotionController (L) FaceButton6"), FKeyDetails::GamepadKey | FKeyDetails::Deprecated));
+	AddKey(FKeyDetails(EKeys::MotionController_Left_FaceButton7, LOCTEXT("MotionController_Left_FaceButton7", "MotionController (L) FaceButton7"), FKeyDetails::GamepadKey | FKeyDetails::Deprecated));
+	AddKey(FKeyDetails(EKeys::MotionController_Left_FaceButton8, LOCTEXT("MotionController_Left_FaceButton8", "MotionController (L) FaceButton8"), FKeyDetails::GamepadKey | FKeyDetails::Deprecated));
+
+	AddKey(FKeyDetails(EKeys::MotionController_Left_Shoulder, LOCTEXT("MotionController_Left_Shoulder", "MotionController (L) Shoulder"), FKeyDetails::GamepadKey | FKeyDetails::Deprecated));
+	AddKey(FKeyDetails(EKeys::MotionController_Left_Trigger, LOCTEXT("MotionController_Left_Trigger", "MotionController (L) Trigger"), FKeyDetails::GamepadKey | FKeyDetails::Deprecated));
+
+	AddKey(FKeyDetails(EKeys::MotionController_Left_Grip1, LOCTEXT("MotionController_Left_Grip1", "MotionController (L) Grip1"), FKeyDetails::GamepadKey | FKeyDetails::Deprecated));
+	AddKey(FKeyDetails(EKeys::MotionController_Left_Grip2, LOCTEXT("MotionController_Left_Grip2", "MotionController (L) Grip2"), FKeyDetails::GamepadKey | FKeyDetails::Deprecated));
+
+	AddKey(FKeyDetails(EKeys::MotionController_Left_Thumbstick, LOCTEXT("MotionController_Left_Thumbstick", "MotionController (L) Thumbstick"), FKeyDetails::GamepadKey | FKeyDetails::Deprecated));
+	AddKey(FKeyDetails(EKeys::MotionController_Left_Thumbstick_Up, LOCTEXT("MotionController_Left_Thumbstick_Up", "MotionController (L) Thumbstick Up"), FKeyDetails::GamepadKey | FKeyDetails::Deprecated));
+	AddKey(FKeyDetails(EKeys::MotionController_Left_Thumbstick_Down, LOCTEXT("MotionController_Left_Thumbstick_Down", "MotionController (L) Thumbstick Down"), FKeyDetails::GamepadKey | FKeyDetails::Deprecated));
+	AddKey(FKeyDetails(EKeys::MotionController_Left_Thumbstick_Left, LOCTEXT("MotionController_Left_Thumbstick_Left", "MotionController (L) Thumbstick Left"), FKeyDetails::GamepadKey | FKeyDetails::Deprecated));
+	AddKey(FKeyDetails(EKeys::MotionController_Left_Thumbstick_Right, LOCTEXT("MotionController_Left_Thumbstick_Right", "MotionController (L) Thumbstick Right"), FKeyDetails::GamepadKey | FKeyDetails::Deprecated));
 
 	//	Right Controller
-	AddKey(FKeyDetails(EKeys::MotionController_Right_FaceButton1, LOCTEXT("MotionController_Right_FaceButton1", "MotionController (R) FaceButton1"), FKeyDetails::GamepadKey | FKeyDetails::NotBlueprintBindableKey | FKeyDetails::NotActionBindableKey));
-	AddKey(FKeyDetails(EKeys::MotionController_Right_FaceButton2, LOCTEXT("MotionController_Right_FaceButton2", "MotionController (R) FaceButton2"), FKeyDetails::GamepadKey | FKeyDetails::NotBlueprintBindableKey | FKeyDetails::NotActionBindableKey));
-	AddKey(FKeyDetails(EKeys::MotionController_Right_FaceButton3, LOCTEXT("MotionController_Right_FaceButton3", "MotionController (R) FaceButton3"), FKeyDetails::GamepadKey | FKeyDetails::NotBlueprintBindableKey | FKeyDetails::NotActionBindableKey));
-	AddKey(FKeyDetails(EKeys::MotionController_Right_FaceButton4, LOCTEXT("MotionController_Right_FaceButton4", "MotionController (R) FaceButton4"), FKeyDetails::GamepadKey | FKeyDetails::NotBlueprintBindableKey | FKeyDetails::NotActionBindableKey));
-	AddKey(FKeyDetails(EKeys::MotionController_Right_FaceButton5, LOCTEXT("MotionController_Right_FaceButton5", "MotionController (R) FaceButton5"), FKeyDetails::GamepadKey | FKeyDetails::NotBlueprintBindableKey | FKeyDetails::NotActionBindableKey));
-	AddKey(FKeyDetails(EKeys::MotionController_Right_FaceButton6, LOCTEXT("MotionController_Right_FaceButton6", "MotionController (R) FaceButton6"), FKeyDetails::GamepadKey | FKeyDetails::NotBlueprintBindableKey | FKeyDetails::NotActionBindableKey));
-	AddKey(FKeyDetails(EKeys::MotionController_Right_FaceButton7, LOCTEXT("MotionController_Right_FaceButton7", "MotionController (R) FaceButton7"), FKeyDetails::GamepadKey | FKeyDetails::NotBlueprintBindableKey | FKeyDetails::NotActionBindableKey));
-	AddKey(FKeyDetails(EKeys::MotionController_Right_FaceButton8, LOCTEXT("MotionController_Right_FaceButton8", "MotionController (R) FaceButton8"), FKeyDetails::GamepadKey | FKeyDetails::NotBlueprintBindableKey | FKeyDetails::NotActionBindableKey));
-
-	AddKey(FKeyDetails(EKeys::MotionController_Right_Shoulder, LOCTEXT("MotionController_Right_Shoulder", "MotionController (R) Shoulder"), FKeyDetails::GamepadKey | FKeyDetails::NotBlueprintBindableKey | FKeyDetails::NotActionBindableKey));
-	AddKey(FKeyDetails(EKeys::MotionController_Right_Trigger, LOCTEXT("MotionController_Right_Trigger", "MotionController (R) Trigger"), FKeyDetails::GamepadKey | FKeyDetails::NotBlueprintBindableKey | FKeyDetails::NotActionBindableKey));
-
-	AddKey(FKeyDetails(EKeys::MotionController_Right_Grip1, LOCTEXT("MotionController_Right_Grip1", "MotionController (R) Grip1"), FKeyDetails::GamepadKey | FKeyDetails::NotBlueprintBindableKey | FKeyDetails::NotActionBindableKey));
-	AddKey(FKeyDetails(EKeys::MotionController_Right_Grip2, LOCTEXT("MotionController_Right_Grip2", "MotionController (R) Grip2"), FKeyDetails::GamepadKey | FKeyDetails::NotBlueprintBindableKey | FKeyDetails::NotActionBindableKey));
-
-	AddKey(FKeyDetails(EKeys::MotionController_Right_Thumbstick, LOCTEXT("MotionController_Right_Thumbstick", "MotionController (R) Thumbstick"), FKeyDetails::GamepadKey | FKeyDetails::NotBlueprintBindableKey | FKeyDetails::NotActionBindableKey));
-	AddKey(FKeyDetails(EKeys::MotionController_Right_Thumbstick_Up, LOCTEXT("MotionController_Right_Thumbstick_Up", "MotionController (R) Thumbstick Up"), FKeyDetails::GamepadKey | FKeyDetails::NotBlueprintBindableKey | FKeyDetails::NotActionBindableKey));
-	AddKey(FKeyDetails(EKeys::MotionController_Right_Thumbstick_Down, LOCTEXT("MotionController_Right_Thumbstick_Down", "MotionController (R) Thumbstick Down"), FKeyDetails::GamepadKey | FKeyDetails::NotBlueprintBindableKey | FKeyDetails::NotActionBindableKey));
-	AddKey(FKeyDetails(EKeys::MotionController_Right_Thumbstick_Left, LOCTEXT("MotionController_Right_Thumbstick_Left", "MotionController (R) Thumbstick Left"), FKeyDetails::GamepadKey | FKeyDetails::NotBlueprintBindableKey | FKeyDetails::NotActionBindableKey));
-	AddKey(FKeyDetails(EKeys::MotionController_Right_Thumbstick_Right, LOCTEXT("MotionController_Right_Thumbstick_Right", "MotionController (R) Thumbstick Right"), FKeyDetails::GamepadKey | FKeyDetails::NotBlueprintBindableKey | FKeyDetails::NotActionBindableKey));
+	AddKey(FKeyDetails(EKeys::MotionController_Right_FaceButton1, LOCTEXT("MotionController_Right_FaceButton1", "MotionController (R) FaceButton1"), FKeyDetails::GamepadKey | FKeyDetails::Deprecated));
+	AddKey(FKeyDetails(EKeys::MotionController_Right_FaceButton2, LOCTEXT("MotionController_Right_FaceButton2", "MotionController (R) FaceButton2"), FKeyDetails::GamepadKey | FKeyDetails::Deprecated));
+	AddKey(FKeyDetails(EKeys::MotionController_Right_FaceButton3, LOCTEXT("MotionController_Right_FaceButton3", "MotionController (R) FaceButton3"), FKeyDetails::GamepadKey | FKeyDetails::Deprecated));
+	AddKey(FKeyDetails(EKeys::MotionController_Right_FaceButton4, LOCTEXT("MotionController_Right_FaceButton4", "MotionController (R) FaceButton4"), FKeyDetails::GamepadKey | FKeyDetails::Deprecated));
+	AddKey(FKeyDetails(EKeys::MotionController_Right_FaceButton5, LOCTEXT("MotionController_Right_FaceButton5", "MotionController (R) FaceButton5"), FKeyDetails::GamepadKey | FKeyDetails::Deprecated));
+	AddKey(FKeyDetails(EKeys::MotionController_Right_FaceButton6, LOCTEXT("MotionController_Right_FaceButton6", "MotionController (R) FaceButton6"), FKeyDetails::GamepadKey | FKeyDetails::Deprecated));
+	AddKey(FKeyDetails(EKeys::MotionController_Right_FaceButton7, LOCTEXT("MotionController_Right_FaceButton7", "MotionController (R) FaceButton7"), FKeyDetails::GamepadKey | FKeyDetails::Deprecated));
+	AddKey(FKeyDetails(EKeys::MotionController_Right_FaceButton8, LOCTEXT("MotionController_Right_FaceButton8", "MotionController (R) FaceButton8"), FKeyDetails::GamepadKey | FKeyDetails::Deprecated));
+
+	AddKey(FKeyDetails(EKeys::MotionController_Right_Shoulder, LOCTEXT("MotionController_Right_Shoulder", "MotionController (R) Shoulder"), FKeyDetails::GamepadKey | FKeyDetails::Deprecated));
+	AddKey(FKeyDetails(EKeys::MotionController_Right_Trigger, LOCTEXT("MotionController_Right_Trigger", "MotionController (R) Trigger"), FKeyDetails::GamepadKey | FKeyDetails::Deprecated));
+
+	AddKey(FKeyDetails(EKeys::MotionController_Right_Grip1, LOCTEXT("MotionController_Right_Grip1", "MotionController (R) Grip1"), FKeyDetails::GamepadKey | FKeyDetails::Deprecated));
+	AddKey(FKeyDetails(EKeys::MotionController_Right_Grip2, LOCTEXT("MotionController_Right_Grip2", "MotionController (R) Grip2"), FKeyDetails::GamepadKey | FKeyDetails::Deprecated));
+
+	AddKey(FKeyDetails(EKeys::MotionController_Right_Thumbstick, LOCTEXT("MotionController_Right_Thumbstick", "MotionController (R) Thumbstick"), FKeyDetails::GamepadKey | FKeyDetails::Deprecated));
+	AddKey(FKeyDetails(EKeys::MotionController_Right_Thumbstick_Up, LOCTEXT("MotionController_Right_Thumbstick_Up", "MotionController (R) Thumbstick Up"), FKeyDetails::GamepadKey | FKeyDetails::Deprecated));
+	AddKey(FKeyDetails(EKeys::MotionController_Right_Thumbstick_Down, LOCTEXT("MotionController_Right_Thumbstick_Down", "MotionController (R) Thumbstick Down"), FKeyDetails::GamepadKey | FKeyDetails::Deprecated));
+	AddKey(FKeyDetails(EKeys::MotionController_Right_Thumbstick_Left, LOCTEXT("MotionController_Right_Thumbstick_Left", "MotionController (R) Thumbstick Left"), FKeyDetails::GamepadKey | FKeyDetails::Deprecated));
+	AddKey(FKeyDetails(EKeys::MotionController_Right_Thumbstick_Right, LOCTEXT("MotionController_Right_Thumbstick_Right", "MotionController (R) Thumbstick Right"), FKeyDetails::GamepadKey | FKeyDetails::Deprecated));
 
 	//   Motion Controller Axes
 	//	Left Controller
-	AddKey(FKeyDetails(EKeys::MotionController_Left_Thumbstick_X, LOCTEXT("MotionController_Left_Thumbstick_X", "MotionController (L) Thumbstick X"), FKeyDetails::GamepadKey | FKeyDetails::FloatAxis | FKeyDetails::NotBlueprintBindableKey | FKeyDetails::NotActionBindableKey));
-	AddKey(FKeyDetails(EKeys::MotionController_Left_Thumbstick_Y, LOCTEXT("MotionController_Left_Thumbstick_Y", "MotionController (L) Thumbstick Y"), FKeyDetails::GamepadKey | FKeyDetails::FloatAxis | FKeyDetails::NotBlueprintBindableKey | FKeyDetails::NotActionBindableKey));
-	AddKey(FKeyDetails(EKeys::MotionController_Left_TriggerAxis, LOCTEXT("MotionController_Left_TriggerAxis", "MotionController (L) TriggerAxis"), FKeyDetails::GamepadKey | FKeyDetails::FloatAxis | FKeyDetails::NotBlueprintBindableKey | FKeyDetails::NotActionBindableKey));
-	AddKey(FKeyDetails(EKeys::MotionController_Left_Grip1Axis, LOCTEXT("MotionController_Left_Grip1Axis", "MotionController (L) Grip1 Axis"), FKeyDetails::GamepadKey | FKeyDetails::FloatAxis | FKeyDetails::NotBlueprintBindableKey | FKeyDetails::NotActionBindableKey));
-	AddKey(FKeyDetails(EKeys::MotionController_Left_Grip2Axis, LOCTEXT("MotionController_Left_Grip2Axis", "MotionController (L) Grip2 Axis"), FKeyDetails::GamepadKey | FKeyDetails::FloatAxis | FKeyDetails::NotBlueprintBindableKey | FKeyDetails::NotActionBindableKey));
+	AddKey(FKeyDetails(EKeys::MotionController_Left_Thumbstick_X, LOCTEXT("MotionController_Left_Thumbstick_X", "MotionController (L) Thumbstick X"), FKeyDetails::GamepadKey | FKeyDetails::FloatAxis | FKeyDetails::Deprecated));
+	AddKey(FKeyDetails(EKeys::MotionController_Left_Thumbstick_Y, LOCTEXT("MotionController_Left_Thumbstick_Y", "MotionController (L) Thumbstick Y"), FKeyDetails::GamepadKey | FKeyDetails::FloatAxis | FKeyDetails::Deprecated));
+	AddKey(FKeyDetails(EKeys::MotionController_Left_TriggerAxis, LOCTEXT("MotionController_Left_TriggerAxis", "MotionController (L) TriggerAxis"), FKeyDetails::GamepadKey | FKeyDetails::FloatAxis | FKeyDetails::Deprecated));
+	AddKey(FKeyDetails(EKeys::MotionController_Left_Grip1Axis, LOCTEXT("MotionController_Left_Grip1Axis", "MotionController (L) Grip1 Axis"), FKeyDetails::GamepadKey | FKeyDetails::FloatAxis | FKeyDetails::Deprecated));
+	AddKey(FKeyDetails(EKeys::MotionController_Left_Grip2Axis, LOCTEXT("MotionController_Left_Grip2Axis", "MotionController (L) Grip2 Axis"), FKeyDetails::GamepadKey | FKeyDetails::FloatAxis | FKeyDetails::Deprecated));
 
 	//	Right Controller
-	AddKey(FKeyDetails(EKeys::MotionController_Right_Thumbstick_X, LOCTEXT("MotionController_Right_Thumbstick_X", "MotionController (R) Thumbstick X"), FKeyDetails::GamepadKey | FKeyDetails::FloatAxis | FKeyDetails::NotBlueprintBindableKey | FKeyDetails::NotActionBindableKey));
-	AddKey(FKeyDetails(EKeys::MotionController_Right_Thumbstick_Y, LOCTEXT("MotionController_Right_Thumbstick_Y", "MotionController (R) Thumbstick Y"), FKeyDetails::GamepadKey | FKeyDetails::FloatAxis | FKeyDetails::NotBlueprintBindableKey | FKeyDetails::NotActionBindableKey));
-	AddKey(FKeyDetails(EKeys::MotionController_Right_TriggerAxis, LOCTEXT("MotionController_Right_TriggerAxis", "MotionController (R) TriggerAxis"), FKeyDetails::GamepadKey | FKeyDetails::FloatAxis | FKeyDetails::NotBlueprintBindableKey | FKeyDetails::NotActionBindableKey));
-	AddKey(FKeyDetails(EKeys::MotionController_Right_Grip1Axis, LOCTEXT("MotionController_Right_Grip1Axis", "MotionController (R) Grip1 Axis"), FKeyDetails::GamepadKey | FKeyDetails::FloatAxis | FKeyDetails::NotBlueprintBindableKey | FKeyDetails::NotActionBindableKey));
-	AddKey(FKeyDetails(EKeys::MotionController_Right_Grip2Axis, LOCTEXT("MotionController_Right_Grip2Axis", "MotionController (R) Grip2 Axis"), FKeyDetails::GamepadKey | FKeyDetails::FloatAxis | FKeyDetails::NotBlueprintBindableKey | FKeyDetails::NotActionBindableKey));
+	AddKey(FKeyDetails(EKeys::MotionController_Right_Thumbstick_X, LOCTEXT("MotionController_Right_Thumbstick_X", "MotionController (R) Thumbstick X"), FKeyDetails::GamepadKey | FKeyDetails::FloatAxis | FKeyDetails::Deprecated));
+	AddKey(FKeyDetails(EKeys::MotionController_Right_Thumbstick_Y, LOCTEXT("MotionController_Right_Thumbstick_Y", "MotionController (R) Thumbstick Y"), FKeyDetails::GamepadKey | FKeyDetails::FloatAxis | FKeyDetails::Deprecated));
+	AddKey(FKeyDetails(EKeys::MotionController_Right_TriggerAxis, LOCTEXT("MotionController_Right_TriggerAxis", "MotionController (R) TriggerAxis"), FKeyDetails::GamepadKey | FKeyDetails::FloatAxis | FKeyDetails::Deprecated));
+	AddKey(FKeyDetails(EKeys::MotionController_Right_Grip1Axis, LOCTEXT("MotionController_Right_Grip1Axis", "MotionController (R) Grip1 Axis"), FKeyDetails::GamepadKey | FKeyDetails::FloatAxis | FKeyDetails::Deprecated));
+	AddKey(FKeyDetails(EKeys::MotionController_Right_Grip2Axis, LOCTEXT("MotionController_Right_Grip2Axis", "MotionController (R) Grip2 Axis"), FKeyDetails::GamepadKey | FKeyDetails::FloatAxis | FKeyDetails::Deprecated));
 
 	// PS4-specific
 	AddMenuCategoryDisplayInfo("PS4", LOCTEXT("PS4SubCategory", "PS4"), TEXT("GraphEditor.PadEvent_16x"));
@@ -882,6 +931,10 @@ void EKeys::Initialize()
 	AddKey(FKeyDetails(EKeys::Vive_Left_Trackpad_Y, LOCTEXT("Vive_Left_Trackpad_Y", "Vive (L) Trackpad Y"), FKeyDetails::GamepadKey | FKeyDetails::FloatAxis | FKeyDetails::NotBlueprintBindableKey, "Vive"));
 	AddKey(FKeyDetails(EKeys::Vive_Left_Trackpad_Click, LOCTEXT("Vive_Left_Trackpad_Click", "Vive (L) Trackpad"), FKeyDetails::GamepadKey | FKeyDetails::NotBlueprintBindableKey, "Vive"));
 	AddKey(FKeyDetails(EKeys::Vive_Left_Trackpad_Touch, LOCTEXT("Vive_Left_Trackpad_Touch", "Vive (L) Trackpad Touch"), FKeyDetails::GamepadKey | FKeyDetails::NotBlueprintBindableKey, "Vive"));
+	AddKey(FKeyDetails(EKeys::Vive_Left_Trackpad_Up, LOCTEXT("Vive_Left_Trackpad_Up", "Vive (L) Trackpad Up"), FKeyDetails::GamepadKey | FKeyDetails::NotBlueprintBindableKey, "Vive"));
+	AddKey(FKeyDetails(EKeys::Vive_Left_Trackpad_Down, LOCTEXT("Vive_Left_Trackpad_Down", "Vive (L) Trackpad Down"), FKeyDetails::GamepadKey | FKeyDetails::NotBlueprintBindableKey, "Vive"));
+	AddKey(FKeyDetails(EKeys::Vive_Left_Trackpad_Left, LOCTEXT("Vive_Left_Trackpad_Left", "Vive (L) Trackpad Left"), FKeyDetails::GamepadKey | FKeyDetails::NotBlueprintBindableKey, "Vive"));
+	AddKey(FKeyDetails(EKeys::Vive_Left_Trackpad_Right, LOCTEXT("Vive_Left_Trackpad_Right", "Vive (L) Trackpad Right"), FKeyDetails::GamepadKey | FKeyDetails::NotBlueprintBindableKey, "Vive"));
 	AddKey(FKeyDetails(EKeys::Vive_Right_System_Click, LOCTEXT("Vive_Right_System_Click", "Vive (R) System"), FKeyDetails::GamepadKey | FKeyDetails::NotBlueprintBindableKey, "Vive"));
 	AddKey(FKeyDetails(EKeys::Vive_Right_Grip_Click, LOCTEXT("Vive_Right_Grip_Click", "Vive (R) Grip"), FKeyDetails::GamepadKey | FKeyDetails::NotBlueprintBindableKey, "Vive"));
 	AddKey(FKeyDetails(EKeys::Vive_Right_Menu_Click, LOCTEXT("Vive_Right_Menu_Click", "Vive (R) Menu"), FKeyDetails::GamepadKey | FKeyDetails::NotBlueprintBindableKey, "Vive"));
@@ -891,6 +944,10 @@ void EKeys::Initialize()
 	AddKey(FKeyDetails(EKeys::Vive_Right_Trackpad_Y, LOCTEXT("Vive_Right_Trackpad_Y", "Vive (R) Trackpad Y"), FKeyDetails::GamepadKey | FKeyDetails::FloatAxis | FKeyDetails::NotBlueprintBindableKey, "Vive"));
 	AddKey(FKeyDetails(EKeys::Vive_Right_Trackpad_Click, LOCTEXT("Vive_Right_Trackpad_Click", "Vive (R) Trackpad"), FKeyDetails::GamepadKey | FKeyDetails::NotBlueprintBindableKey, "Vive"));
 	AddKey(FKeyDetails(EKeys::Vive_Right_Trackpad_Touch, LOCTEXT("Vive_Right_Trackpad_Touch", "Vive (R) Trackpad Touch"), FKeyDetails::GamepadKey | FKeyDetails::NotBlueprintBindableKey, "Vive"));
+	AddKey(FKeyDetails(EKeys::Vive_Right_Trackpad_Up, LOCTEXT("Vive_Right_Trackpad_Up", "Vive (R) Trackpad Up"), FKeyDetails::GamepadKey | FKeyDetails::NotBlueprintBindableKey, "Vive"));
+	AddKey(FKeyDetails(EKeys::Vive_Right_Trackpad_Down, LOCTEXT("Vive_Right_Trackpad_Down", "Vive (R) Trackpad Down"), FKeyDetails::GamepadKey | FKeyDetails::NotBlueprintBindableKey, "Vive"));
+	AddKey(FKeyDetails(EKeys::Vive_Right_Trackpad_Left, LOCTEXT("Vive_Right_Trackpad_Left", "Vive (R) Trackpad Left"), FKeyDetails::GamepadKey | FKeyDetails::NotBlueprintBindableKey, "Vive"));
+	AddKey(FKeyDetails(EKeys::Vive_Right_Trackpad_Right, LOCTEXT("Vive_Right_Trackpad_Right", "Vive (R) Trackpad Right"), FKeyDetails::GamepadKey | FKeyDetails::NotBlueprintBindableKey, "Vive"));
 
 	// Microsoft Mixed Reality Motion Controller
 	AddMenuCategoryDisplayInfo("MixedReality", LOCTEXT("MixedRealitySubCategory", "Windows Mixed Reality"), TEXT("GraphEditor.PadEvent_16x"));
@@ -902,10 +959,18 @@ void EKeys::Initialize()
 	AddKey(FKeyDetails(EKeys::MixedReality_Left_Thumbstick_X, LOCTEXT("MixedReality_Left_Thumbstick_X", "Mixed Reality (L) Thumbstick X"), FKeyDetails::GamepadKey | FKeyDetails::FloatAxis | FKeyDetails::NotBlueprintBindableKey, "MixedReality"));
 	AddKey(FKeyDetails(EKeys::MixedReality_Left_Thumbstick_Y, LOCTEXT("MixedReality_Left_Thumbstick_Y", "Mixed Reality (L) Thumbstick Y"), FKeyDetails::GamepadKey | FKeyDetails::FloatAxis | FKeyDetails::NotBlueprintBindableKey, "MixedReality"));
 	AddKey(FKeyDetails(EKeys::MixedReality_Left_Thumbstick_Click, LOCTEXT("MixedReality_Left_Thumbstick_Click", "Mixed Reality (L) Thumbstick"), FKeyDetails::GamepadKey | FKeyDetails::NotBlueprintBindableKey, "MixedReality"));
+	AddKey(FKeyDetails(EKeys::MixedReality_Left_Thumbstick_Up, LOCTEXT("MixedReality_Left_Thumbstick_Up", "Mixed Reality (L) Thumbstick Up"), FKeyDetails::GamepadKey | FKeyDetails::NotBlueprintBindableKey, "MixedReality"));
+	AddKey(FKeyDetails(EKeys::MixedReality_Left_Thumbstick_Down, LOCTEXT("MixedReality_Left_Thumbstick_Down", "Mixed Reality (L) Thumbstick Down"), FKeyDetails::GamepadKey | FKeyDetails::NotBlueprintBindableKey, "MixedReality"));
+	AddKey(FKeyDetails(EKeys::MixedReality_Left_Thumbstick_Left, LOCTEXT("MixedReality_Left_Thumbstick_Left", "Mixed Reality (L) Thumbstick Left"), FKeyDetails::GamepadKey | FKeyDetails::NotBlueprintBindableKey, "MixedReality"));
+	AddKey(FKeyDetails(EKeys::MixedReality_Left_Thumbstick_Right, LOCTEXT("MixedReality_Left_Thumbstick_Right", "Mixed Reality (L) Thumbstick Right"), FKeyDetails::GamepadKey | FKeyDetails::NotBlueprintBindableKey, "MixedReality"));
 	AddKey(FKeyDetails(EKeys::MixedReality_Left_Trackpad_X, LOCTEXT("MixedReality_Left_Trackpad_X", "Mixed Reality (L) Trackpad X"), FKeyDetails::GamepadKey | FKeyDetails::FloatAxis | FKeyDetails::NotBlueprintBindableKey, "MixedReality"));
 	AddKey(FKeyDetails(EKeys::MixedReality_Left_Trackpad_Y, LOCTEXT("MixedReality_Left_Trackpad_Y", "Mixed Reality (L) Trackpad Y"), FKeyDetails::GamepadKey | FKeyDetails::FloatAxis | FKeyDetails::NotBlueprintBindableKey, "MixedReality"));
 	AddKey(FKeyDetails(EKeys::MixedReality_Left_Trackpad_Click, LOCTEXT("MixedReality_Left_Trackpad_Click", "Mixed Reality (L) Trackpad"), FKeyDetails::GamepadKey | FKeyDetails::NotBlueprintBindableKey, "MixedReality"));
 	AddKey(FKeyDetails(EKeys::MixedReality_Left_Trackpad_Touch, LOCTEXT("MixedReality_Left_Trackpad_Touch", "Mixed Reality (L) Trackpad Touch"), FKeyDetails::GamepadKey | FKeyDetails::NotBlueprintBindableKey, "MixedReality"));
+	AddKey(FKeyDetails(EKeys::MixedReality_Left_Trackpad_Up, LOCTEXT("MixedReality_Left_Trackpad_Up", "Mixed Reality (L) Trackpad Up"), FKeyDetails::GamepadKey | FKeyDetails::NotBlueprintBindableKey, "MixedReality"));
+	AddKey(FKeyDetails(EKeys::MixedReality_Left_Trackpad_Down, LOCTEXT("MixedReality_Left_Trackpad_Down", "Mixed Reality (L) Trackpad Down"), FKeyDetails::GamepadKey | FKeyDetails::NotBlueprintBindableKey, "MixedReality"));
+	AddKey(FKeyDetails(EKeys::MixedReality_Left_Trackpad_Left, LOCTEXT("MixedReality_Left_Trackpad_Left", "Mixed Reality (L) Trackpad Left"), FKeyDetails::GamepadKey | FKeyDetails::NotBlueprintBindableKey, "MixedReality"));
+	AddKey(FKeyDetails(EKeys::MixedReality_Left_Trackpad_Right, LOCTEXT("MixedReality_Left_Trackpad_Right", "Mixed Reality (L) Trackpad Right"), FKeyDetails::GamepadKey | FKeyDetails::NotBlueprintBindableKey, "MixedReality"));
 	AddKey(FKeyDetails(EKeys::MixedReality_Right_Menu_Click, LOCTEXT("MixedReality_Right_Menu_Click", "Mixed Reality (R) Menu"), FKeyDetails::GamepadKey | FKeyDetails::NotBlueprintBindableKey, "MixedReality"));
 	AddKey(FKeyDetails(EKeys::MixedReality_Right_Grip_Click, LOCTEXT("MixedReality_Right_Grip_Click", "Mixed Reality (R) Grip"), FKeyDetails::GamepadKey | FKeyDetails::NotBlueprintBindableKey, "MixedReality"));
 	AddKey(FKeyDetails(EKeys::MixedReality_Right_Trigger_Click, LOCTEXT("MixedReality_Right_Trigger_Click", "Mixed Reality (R) Trigger"), FKeyDetails::GamepadKey | FKeyDetails::NotBlueprintBindableKey, "MixedReality"));
@@ -913,10 +978,18 @@ void EKeys::Initialize()
 	AddKey(FKeyDetails(EKeys::MixedReality_Right_Thumbstick_X, LOCTEXT("MixedReality_Right_Thumbstick_X", "Mixed Reality (R) Thumbstick X"), FKeyDetails::GamepadKey | FKeyDetails::FloatAxis | FKeyDetails::NotBlueprintBindableKey, "MixedReality"));
 	AddKey(FKeyDetails(EKeys::MixedReality_Right_Thumbstick_Y, LOCTEXT("MixedReality_Right_Thumbstick_Y", "Mixed Reality (R) Thumbstick Y"), FKeyDetails::GamepadKey | FKeyDetails::FloatAxis | FKeyDetails::NotBlueprintBindableKey, "MixedReality"));
 	AddKey(FKeyDetails(EKeys::MixedReality_Right_Thumbstick_Click, LOCTEXT("MixedReality_Right_Thumbstick_Click", "Mixed Reality (R) Thumbstick"), FKeyDetails::GamepadKey | FKeyDetails::NotBlueprintBindableKey, "MixedReality"));
+	AddKey(FKeyDetails(EKeys::MixedReality_Right_Thumbstick_Up, LOCTEXT("MixedReality_Right_Thumbstick_Up", "Mixed Reality (R) Thumbstick Up"), FKeyDetails::GamepadKey | FKeyDetails::NotBlueprintBindableKey, "MixedReality"));
+	AddKey(FKeyDetails(EKeys::MixedReality_Right_Thumbstick_Down, LOCTEXT("MixedReality_Right_Thumbstick_Down", "Mixed Reality (R) Thumbstick Down"), FKeyDetails::GamepadKey | FKeyDetails::NotBlueprintBindableKey, "MixedReality"));
+	AddKey(FKeyDetails(EKeys::MixedReality_Right_Thumbstick_Left, LOCTEXT("MixedReality_Right_Thumbstick_Left", "Mixed Reality (R) Thumbstick Left"), FKeyDetails::GamepadKey | FKeyDetails::NotBlueprintBindableKey, "MixedReality"));
+	AddKey(FKeyDetails(EKeys::MixedReality_Right_Thumbstick_Right, LOCTEXT("MixedReality_Right_Thumbstick_Right", "Mixed Reality (R) Thumbstick Right"), FKeyDetails::GamepadKey | FKeyDetails::NotBlueprintBindableKey, "MixedReality"));
 	AddKey(FKeyDetails(EKeys::MixedReality_Right_Trackpad_X, LOCTEXT("MixedReality_Right_Trackpad_X", "Mixed Reality (R) Trackpad X"), FKeyDetails::GamepadKey | FKeyDetails::FloatAxis | FKeyDetails::NotBlueprintBindableKey, "MixedReality"));
 	AddKey(FKeyDetails(EKeys::MixedReality_Right_Trackpad_Y, LOCTEXT("MixedReality_Right_Trackpad_Y", "Mixed Reality (R) Trackpad Y"), FKeyDetails::GamepadKey | FKeyDetails::FloatAxis | FKeyDetails::NotBlueprintBindableKey, "MixedReality"));
 	AddKey(FKeyDetails(EKeys::MixedReality_Right_Trackpad_Click, LOCTEXT("MixedReality_Right_Trackpad_Click", "Mixed Reality (R) Trackpad"), FKeyDetails::GamepadKey | FKeyDetails::NotBlueprintBindableKey, "MixedReality"));
 	AddKey(FKeyDetails(EKeys::MixedReality_Right_Trackpad_Touch, LOCTEXT("MixedReality_Right_Trackpad_Touch", "Mixed Reality (R) Trackpad Touch"), FKeyDetails::GamepadKey | FKeyDetails::NotBlueprintBindableKey, "MixedReality"));
+	AddKey(FKeyDetails(EKeys::MixedReality_Right_Trackpad_Up, LOCTEXT("MixedReality_Right_Trackpad_Up", "Mixed Reality (R) Trackpad Up"), FKeyDetails::GamepadKey | FKeyDetails::NotBlueprintBindableKey, "MixedReality"));
+	AddKey(FKeyDetails(EKeys::MixedReality_Right_Trackpad_Down, LOCTEXT("MixedReality_Right_Trackpad_Down", "Mixed Reality (R) Trackpad Down"), FKeyDetails::GamepadKey | FKeyDetails::NotBlueprintBindableKey, "MixedReality"));
+	AddKey(FKeyDetails(EKeys::MixedReality_Right_Trackpad_Left, LOCTEXT("MixedReality_Right_Trackpad_Left", "Mixed Reality (R) Trackpad Left"), FKeyDetails::GamepadKey | FKeyDetails::NotBlueprintBindableKey, "MixedReality"));
+	AddKey(FKeyDetails(EKeys::MixedReality_Right_Trackpad_Right, LOCTEXT("MixedReality_Right_Trackpad_Right", "Mixed Reality (R) Trackpad Right"), FKeyDetails::GamepadKey | FKeyDetails::NotBlueprintBindableKey, "MixedReality"));
 
 	// Oculus Go Controller
 	AddMenuCategoryDisplayInfo("OculusGo", LOCTEXT("OculusGoSubCategory", "Oculus Go"), TEXT("GraphEditor.PadEvent_16x"));
@@ -953,6 +1026,10 @@ void EKeys::Initialize()
 	AddKey(FKeyDetails(EKeys::OculusTouch_Left_Thumbstick_Y, LOCTEXT("OculusTouch_Left_Thumbstick_Y", "Oculus Touch (L) Thumbstick Y"), FKeyDetails::GamepadKey | FKeyDetails::FloatAxis | FKeyDetails::NotBlueprintBindableKey, "OculusTouch"));
 	AddKey(FKeyDetails(EKeys::OculusTouch_Left_Thumbstick_Click, LOCTEXT("OculusTouch_Left_Thumbstick_Click", "Oculus Touch (L) Thumbstick"), FKeyDetails::GamepadKey | FKeyDetails::NotBlueprintBindableKey, "OculusTouch"));
 	AddKey(FKeyDetails(EKeys::OculusTouch_Left_Thumbstick_Touch, LOCTEXT("OculusTouch_Left_Thumbstick_Touch", "Oculus Touch (L) Thumbstick Touch"), FKeyDetails::GamepadKey | FKeyDetails::NotBlueprintBindableKey, "OculusTouch"));
+	AddKey(FKeyDetails(EKeys::OculusTouch_Left_Thumbstick_Up, LOCTEXT("OculusTouch_Left_Thumbstick_Up", "Oculus Touch (L) Thumbstick Up"), FKeyDetails::GamepadKey | FKeyDetails::NotBlueprintBindableKey, "OculusTouch"));
+	AddKey(FKeyDetails(EKeys::OculusTouch_Left_Thumbstick_Down, LOCTEXT("OculusTouch_Left_Thumbstick_Down", "Oculus Touch (L) Thumbstick Down"), FKeyDetails::GamepadKey | FKeyDetails::NotBlueprintBindableKey, "OculusTouch"));
+	AddKey(FKeyDetails(EKeys::OculusTouch_Left_Thumbstick_Left, LOCTEXT("OculusTouch_Left_Thumbstick_Left", "Oculus Touch (L) Thumbstick Left"), FKeyDetails::GamepadKey | FKeyDetails::NotBlueprintBindableKey, "OculusTouch"));
+	AddKey(FKeyDetails(EKeys::OculusTouch_Left_Thumbstick_Right, LOCTEXT("OculusTouch_Left_Thumbstick_Right", "Oculus Touch (L) Thumbstick Right"), FKeyDetails::GamepadKey | FKeyDetails::NotBlueprintBindableKey, "OculusTouch"));
 	AddKey(FKeyDetails(EKeys::OculusTouch_Right_A_Click, LOCTEXT("OculusTouch_Right_A_Click", "Oculus Touch (R) A Press"), FKeyDetails::GamepadKey | FKeyDetails::NotBlueprintBindableKey, "OculusTouch"));
 	AddKey(FKeyDetails(EKeys::OculusTouch_Right_B_Click, LOCTEXT("OculusTouch_Right_B_Click", "Oculus Touch (R) B Press"), FKeyDetails::GamepadKey | FKeyDetails::NotBlueprintBindableKey, "OculusTouch"));
 	AddKey(FKeyDetails(EKeys::OculusTouch_Right_A_Touch, LOCTEXT("OculusTouch_Right_A_Touch", "Oculus Touch (R) A Touch"), FKeyDetails::GamepadKey | FKeyDetails::NotBlueprintBindableKey, "OculusTouch"));
@@ -967,6 +1044,10 @@ void EKeys::Initialize()
 	AddKey(FKeyDetails(EKeys::OculusTouch_Right_Thumbstick_Y, LOCTEXT("OculusTouch_Right_Thumbstick_Y", "Oculus Touch (R) Thumbstick Y"), FKeyDetails::GamepadKey | FKeyDetails::FloatAxis | FKeyDetails::NotBlueprintBindableKey, "OculusTouch"));
 	AddKey(FKeyDetails(EKeys::OculusTouch_Right_Thumbstick_Click, LOCTEXT("OculusTouch_Right_Thumbstick_Click", "Oculus Touch (R) Thumbstick"), FKeyDetails::GamepadKey | FKeyDetails::NotBlueprintBindableKey, "OculusTouch"));
 	AddKey(FKeyDetails(EKeys::OculusTouch_Right_Thumbstick_Touch, LOCTEXT("OculusTouch_Right_Thumbstick_Touch", "Oculus Touch (R) Thumbstick Touch"), FKeyDetails::GamepadKey | FKeyDetails::NotBlueprintBindableKey, "OculusTouch"));
+	AddKey(FKeyDetails(EKeys::OculusTouch_Right_Thumbstick_Up, LOCTEXT("OculusTouch_Right_Thumbstick_Up", "Oculus Touch (R) Thumbstick Up"), FKeyDetails::GamepadKey | FKeyDetails::NotBlueprintBindableKey, "OculusTouch"));
+	AddKey(FKeyDetails(EKeys::OculusTouch_Right_Thumbstick_Down, LOCTEXT("OculusTouch_Right_Thumbstick_Down", "Oculus Touch (R) Thumbstick Down"), FKeyDetails::GamepadKey | FKeyDetails::NotBlueprintBindableKey, "OculusTouch"));
+	AddKey(FKeyDetails(EKeys::OculusTouch_Right_Thumbstick_Left, LOCTEXT("OculusTouch_Right_Thumbstick_Left", "Oculus Touch (R) Thumbstick Left"), FKeyDetails::GamepadKey | FKeyDetails::NotBlueprintBindableKey, "OculusTouch"));
+	AddKey(FKeyDetails(EKeys::OculusTouch_Right_Thumbstick_Right, LOCTEXT("OculusTouch_Right_Thumbstick_Right", "Oculus Touch (R) Thumbstick Right"), FKeyDetails::GamepadKey | FKeyDetails::NotBlueprintBindableKey, "OculusTouch"));
 
 	// Valve Index Controller
 	AddMenuCategoryDisplayInfo("ValveIndex", LOCTEXT("ValveIndexSubCategory", "Valve Index"), TEXT("GraphEditor.PadEvent_16x"));
@@ -987,11 +1068,19 @@ void EKeys::Initialize()
 	AddKey(FKeyDetails(EKeys::ValveIndex_Left_Thumbstick_Y, LOCTEXT("ValveIndex_Left_Thumbstick_Y", "Valve Index (L) Thumbstick Y"), FKeyDetails::GamepadKey | FKeyDetails::FloatAxis | FKeyDetails::NotBlueprintBindableKey, "ValveIndex"));
 	AddKey(FKeyDetails(EKeys::ValveIndex_Left_Thumbstick_Click, LOCTEXT("ValveIndex_Left_Thumbstick_Click", "Valve Index (L) Thumbstick"), FKeyDetails::GamepadKey | FKeyDetails::NotBlueprintBindableKey, "ValveIndex"));
 	AddKey(FKeyDetails(EKeys::ValveIndex_Left_Thumbstick_Touch, LOCTEXT("ValveIndex_Left_Thumbstick_Touch", "Valve Index (L) Thumbstick Touch"), FKeyDetails::GamepadKey | FKeyDetails::NotBlueprintBindableKey, "ValveIndex"));
+	AddKey(FKeyDetails(EKeys::ValveIndex_Left_Thumbstick_Up, LOCTEXT("ValveIndex_Left_Thumbstick_Up", "Valve Index (L) Thumbstick Up"), FKeyDetails::GamepadKey | FKeyDetails::NotBlueprintBindableKey, "ValveIndex"));
+	AddKey(FKeyDetails(EKeys::ValveIndex_Left_Thumbstick_Down, LOCTEXT("ValveIndex_Left_Thumbstick_Down", "Valve Index (L) Thumbstick Down"), FKeyDetails::GamepadKey | FKeyDetails::NotBlueprintBindableKey, "ValveIndex"));
+	AddKey(FKeyDetails(EKeys::ValveIndex_Left_Thumbstick_Left, LOCTEXT("ValveIndex_Left_Thumbstick_Left", "Valve Index (L) Thumbstick Left"), FKeyDetails::GamepadKey | FKeyDetails::NotBlueprintBindableKey, "ValveIndex"));
+	AddKey(FKeyDetails(EKeys::ValveIndex_Left_Thumbstick_Right, LOCTEXT("ValveIndex_Left_Thumbstick_Right", "Valve Index (L) Thumbstick Right"), FKeyDetails::GamepadKey | FKeyDetails::NotBlueprintBindableKey, "ValveIndex"));
 	AddKey(FKeyDetails(EKeys::ValveIndex_Left_Trackpad_X, LOCTEXT("ValveIndex_Left_Trackpad_X", "Valve Index (L) Trackpad X"), FKeyDetails::GamepadKey | FKeyDetails::FloatAxis | FKeyDetails::NotBlueprintBindableKey, "ValveIndex"));
 	AddKey(FKeyDetails(EKeys::ValveIndex_Left_Trackpad_Y, LOCTEXT("ValveIndex_Left_Trackpad_Y", "Valve Index (L) Trackpad Y"), FKeyDetails::GamepadKey | FKeyDetails::FloatAxis | FKeyDetails::NotBlueprintBindableKey, "ValveIndex"));
 	AddKey(FKeyDetails(EKeys::ValveIndex_Left_Trackpad_Click, LOCTEXT("ValveIndex_Left_Trackpad_Click", "Valve Index (L) Trackpad"), FKeyDetails::GamepadKey | FKeyDetails::NotBlueprintBindableKey, "ValveIndex"));
 	AddKey(FKeyDetails(EKeys::ValveIndex_Left_Trackpad_Force, LOCTEXT("ValveIndex_Left_Trackpad_Force", "Valve Index (L) Trackpad Force"), FKeyDetails::GamepadKey | FKeyDetails::FloatAxis | FKeyDetails::NotBlueprintBindableKey, "ValveIndex"));
 	AddKey(FKeyDetails(EKeys::ValveIndex_Left_Trackpad_Touch, LOCTEXT("ValveIndex_Left_Trackpad_Touch", "Valve Index (L) Trackpad Touch"), FKeyDetails::GamepadKey | FKeyDetails::FloatAxis | FKeyDetails::NotBlueprintBindableKey, "ValveIndex"));
+	AddKey(FKeyDetails(EKeys::ValveIndex_Left_Trackpad_Up, LOCTEXT("ValveIndex_Left_Trackpad_Up", "Valve Index (L) Trackpad Up"), FKeyDetails::GamepadKey | FKeyDetails::NotBlueprintBindableKey, "ValveIndex"));
+	AddKey(FKeyDetails(EKeys::ValveIndex_Left_Trackpad_Down, LOCTEXT("ValveIndex_Left_Trackpad_Down", "Valve Index (L) Trackpad Down"), FKeyDetails::GamepadKey | FKeyDetails::NotBlueprintBindableKey, "ValveIndex"));
+	AddKey(FKeyDetails(EKeys::ValveIndex_Left_Trackpad_Left, LOCTEXT("ValveIndex_Left_Trackpad_Left", "Valve Index (L) Trackpad Left"), FKeyDetails::GamepadKey | FKeyDetails::NotBlueprintBindableKey, "ValveIndex"));
+	AddKey(FKeyDetails(EKeys::ValveIndex_Left_Trackpad_Right, LOCTEXT("ValveIndex_Left_Trackpad_Right", "Valve Index (L) Trackpad Right"), FKeyDetails::GamepadKey | FKeyDetails::NotBlueprintBindableKey, "ValveIndex"));
 	AddKey(FKeyDetails(EKeys::ValveIndex_Right_A_Click, LOCTEXT("ValveIndex_Right_A_Click", "Valve Index (R) A Press"), FKeyDetails::GamepadKey | FKeyDetails::NotBlueprintBindableKey, "ValveIndex"));
 	AddKey(FKeyDetails(EKeys::ValveIndex_Right_B_Click, LOCTEXT("ValveIndex_Right_B_Click", "Valve Index (R) B Press"), FKeyDetails::GamepadKey | FKeyDetails::NotBlueprintBindableKey, "ValveIndex"));
 	AddKey(FKeyDetails(EKeys::ValveIndex_Right_A_Touch, LOCTEXT("ValveIndex_Right_A_Touch", "Valve Index (R) A Touch"), FKeyDetails::GamepadKey | FKeyDetails::NotBlueprintBindableKey, "ValveIndex"));
@@ -1008,11 +1097,19 @@ void EKeys::Initialize()
 	AddKey(FKeyDetails(EKeys::ValveIndex_Right_Thumbstick_Y, LOCTEXT("ValveIndex_Right_Thumbstick_Y", "Valve Index (R) Thumbstick Y"), FKeyDetails::GamepadKey | FKeyDetails::FloatAxis | FKeyDetails::NotBlueprintBindableKey, "ValveIndex"));
 	AddKey(FKeyDetails(EKeys::ValveIndex_Right_Thumbstick_Click, LOCTEXT("ValveIndex_Right_Thumbstick_Click", "Valve Index (R) Thumbstick"), FKeyDetails::GamepadKey | FKeyDetails::NotBlueprintBindableKey, "ValveIndex"));
 	AddKey(FKeyDetails(EKeys::ValveIndex_Right_Thumbstick_Touch, LOCTEXT("ValveIndex_Right_Thumbstick_Touch", "Valve Index (R) Thumbstick Touch"), FKeyDetails::GamepadKey | FKeyDetails::NotBlueprintBindableKey, "ValveIndex"));
+	AddKey(FKeyDetails(EKeys::ValveIndex_Right_Thumbstick_Up, LOCTEXT("ValveIndex_Right_Thumbstick_Up", "Valve Index (R) Thumbstick Up"), FKeyDetails::GamepadKey | FKeyDetails::NotBlueprintBindableKey, "ValveIndex"));
+	AddKey(FKeyDetails(EKeys::ValveIndex_Right_Thumbstick_Down, LOCTEXT("ValveIndex_Right_Thumbstick_Down", "Valve Index (R) Thumbstick Down"), FKeyDetails::GamepadKey | FKeyDetails::NotBlueprintBindableKey, "ValveIndex"));
+	AddKey(FKeyDetails(EKeys::ValveIndex_Right_Thumbstick_Left, LOCTEXT("ValveIndex_Right_Thumbstick_Left", "Valve Index (R) Thumbstick Left"), FKeyDetails::GamepadKey | FKeyDetails::NotBlueprintBindableKey, "ValveIndex"));
+	AddKey(FKeyDetails(EKeys::ValveIndex_Right_Thumbstick_Right, LOCTEXT("ValveIndex_Right_Thumbstick_Right", "Valve Index (R) Thumbstick Right"), FKeyDetails::GamepadKey | FKeyDetails::NotBlueprintBindableKey, "ValveIndex"));
 	AddKey(FKeyDetails(EKeys::ValveIndex_Right_Trackpad_X, LOCTEXT("ValveIndex_Right_Trackpad_X", "Valve Index (R) Trackpad X"), FKeyDetails::GamepadKey | FKeyDetails::FloatAxis | FKeyDetails::NotBlueprintBindableKey, "ValveIndex"));
 	AddKey(FKeyDetails(EKeys::ValveIndex_Right_Trackpad_Y, LOCTEXT("ValveIndex_Right_Trackpad_Y", "Valve Index (R) Trackpad Y"), FKeyDetails::GamepadKey | FKeyDetails::FloatAxis | FKeyDetails::NotBlueprintBindableKey, "ValveIndex"));
 	AddKey(FKeyDetails(EKeys::ValveIndex_Right_Trackpad_Click, LOCTEXT("ValveIndex_Right_Trackpad_Click", "Valve Index (R) Trackpad"), FKeyDetails::GamepadKey | FKeyDetails::NotBlueprintBindableKey, "ValveIndex"));
 	AddKey(FKeyDetails(EKeys::ValveIndex_Right_Trackpad_Force, LOCTEXT("ValveIndex_Right_Trackpad_Force", "Valve Index (R) Trackpad Force"), FKeyDetails::GamepadKey | FKeyDetails::FloatAxis | FKeyDetails::NotBlueprintBindableKey, "ValveIndex"));
 	AddKey(FKeyDetails(EKeys::ValveIndex_Right_Trackpad_Touch, LOCTEXT("ValveIndex_Right_Trackpad_Touch", "Valve Index (R) Trackpad Touch"), FKeyDetails::GamepadKey | FKeyDetails::NotBlueprintBindableKey, "ValveIndex"));
+	AddKey(FKeyDetails(EKeys::ValveIndex_Right_Trackpad_Up, LOCTEXT("ValveIndex_Right_Trackpad_Up", "Valve Index (R) Trackpad Up"), FKeyDetails::GamepadKey | FKeyDetails::NotBlueprintBindableKey, "ValveIndex"));
+	AddKey(FKeyDetails(EKeys::ValveIndex_Right_Trackpad_Down, LOCTEXT("ValveIndex_Right_Trackpad_Down", "Valve Index (R) Trackpad Down"), FKeyDetails::GamepadKey | FKeyDetails::NotBlueprintBindableKey, "ValveIndex"));
+	AddKey(FKeyDetails(EKeys::ValveIndex_Right_Trackpad_Left, LOCTEXT("ValveIndex_Right_Trackpad_Left", "Valve Index (R) Trackpad Left"), FKeyDetails::GamepadKey | FKeyDetails::NotBlueprintBindableKey, "ValveIndex"));
+	AddKey(FKeyDetails(EKeys::ValveIndex_Right_Trackpad_Right, LOCTEXT("ValveIndex_Right_Trackpad_Right", "Valve Index (R) Trackpad Right"), FKeyDetails::GamepadKey | FKeyDetails::NotBlueprintBindableKey, "ValveIndex"));
 
 	// Initialize the input key manager.  This will cause any additional OEM keys to get added
 	FInputKeyManager::Get();
@@ -1350,6 +1447,12 @@ bool FKey::IsBindableToActions() const
 	return (KeyDetails.IsValid() ? KeyDetails->IsBindableToActions() : false);
 }
 
+bool FKey::IsDeprecated() const
+{
+	ConditionalLookupKeyDetails();
+	return (KeyDetails.IsValid() ? KeyDetails->IsDeprecated() : false);
+}
+
 FText FKeyDetails::GetDisplayName(const bool bLongDisplayName) const
 {
 	return ((bLongDisplayName || !ShortDisplayName.IsSet()) ? LongDisplayName.Get() : ShortDisplayName.Get());
