Deprecated console specific Fkeys in favor of generic ones

#rnx

diff --git a/Engine/Source/Runtime/InputCore/Classes/InputCoreTypes.h b/Engine/Source/Runtime/InputCore/Classes/InputCoreTypes.h
--- a/Engine/Source/Runtime/InputCore/Classes/InputCoreTypes.h
+++ b/Engine/Source/Runtime/InputCore/Classes/InputCoreTypes.h
@@ -472,6 +472,7 @@ struct INPUTCORE_API EKeys
 	static const FKey Gesture_Rotate;
 
 	// PS4-specific
+	UE_DEPRECATED(5.0, "This key has deprecated and will be removed. Use GamePad_Special_Left/Right instead.")
 	static const FKey PS4_Special;
 
 	// Steam Controller Specific
diff --git a/Engine/Source/Runtime/InputCore/Private/InputCoreTypes.cpp b/Engine/Source/Runtime/InputCore/Private/InputCoreTypes.cpp
--- a/Engine/Source/Runtime/InputCore/Private/InputCoreTypes.cpp
+++ b/Engine/Source/Runtime/InputCore/Private/InputCoreTypes.cpp
@@ -231,8 +231,10 @@ const FKey EKeys::Gesture_Pinch("Gesture_Pinch");
 const FKey EKeys::Gesture_Flick("Gesture_Flick");
 const FKey EKeys::Gesture_Rotate("Gesture_Rotate");
 
+PRAGMA_DISABLE_DEPRECATION_WARNINGS
 // PS4-specific
 const FKey EKeys::PS4_Special("PS4_Special");
+PRAGMA_ENABLE_DEPRECATION_WARNINGS
 
 // Steam Controller Specific
 const FKey EKeys::Steam_Touch_0("Steam_Touch_0");
@@ -734,10 +736,12 @@ void EKeys::Initialize()
 	AddKey(FKeyDetails(EKeys::Gesture_Rotate, LOCTEXT("Gesture_Rotate", "Rotate"), 0, "Gesture"));
 
 	// PS4-specific
-	AddMenuCategoryDisplayInfo("PS4", LOCTEXT("PS4SubCategory", "PS4"), TEXT("GraphEditor.PadEvent_16x"));
-	AddKey(FKeyDetails(EKeys::PS4_Special, LOCTEXT("PS4_Special", "PS4_Special"), FKeyDetails::NotBlueprintBindableKey | FKeyDetails::NotActionBindableKey, "PS4"));
-	AddKey(FKeyDetails(EKeys::Gamepad_Special_Left_X, LOCTEXT("PS4_Gamepad_Special_Left_X", "PS4 Touchpad Button X Axis"), FKeyDetails::GamepadKey | FKeyDetails::Axis1D, "PS4"));
-	AddKey(FKeyDetails(EKeys::Gamepad_Special_Left_Y, LOCTEXT("PS4_Gamepad_Special_Left_Y", "PS4 Touchpad Button Y Axis"), FKeyDetails::GamepadKey | FKeyDetails::Axis1D, "PS4"));
+	AddMenuCategoryDisplayInfo("Special Gamepad", LOCTEXT("SpecialGamepadSubCategory", "SpecialGamepad"), TEXT("GraphEditor.PadEvent_16x"));
+	PRAGMA_DISABLE_DEPRECATION_WARNINGS
+	AddKey(FKeyDetails(EKeys::PS4_Special, LOCTEXT("PS4_Special", "PS4_Special_DEPRECATED"), FKeyDetails::Deprecated | FKeyDetails::NotBlueprintBindableKey | FKeyDetails::NotActionBindableKey, "PS4"));
+	PRAGMA_ENABLE_DEPRECATION_WARNINGS
+	AddKey(FKeyDetails(EKeys::Gamepad_Special_Left_X, LOCTEXT("Gamepad_Special_Left_X", "Touchpad Button X Axis"), FKeyDetails::GamepadKey | FKeyDetails::Axis1D, "PS4"));
+	AddKey(FKeyDetails(EKeys::Gamepad_Special_Left_Y, LOCTEXT("Gamepad_Special_Left_Y", "Touchpad Button Y Axis"), FKeyDetails::GamepadKey | FKeyDetails::Axis1D, "PS4"));
 
 
 	// Steam Controller specific
diff --git a/Engine/Source/Runtime/Slate/Private/Widgets/Input/SInputKeySelector.cpp b/Engine/Source/Runtime/Slate/Private/Widgets/Input/SInputKeySelector.cpp
--- a/Engine/Source/Runtime/Slate/Private/Widgets/Input/SInputKeySelector.cpp
+++ b/Engine/Source/Runtime/Slate/Private/Widgets/Input/SInputKeySelector.cpp
@@ -180,8 +180,7 @@ FReply SInputKeySelector::OnKeyUp( const FGeometry& MyGeometry, const FKeyEvent&
 	{
 		SetIsSelectingKey( false );
 
-		if ((InKeyEvent.GetKey() == EKeys::PS4_Special) || // Required?
-			(bEscapeCancelsSelection && (KeyUp == EKeys::Escape || IsEscapeKey(KeyUp))))
+		if (bEscapeCancelsSelection && (KeyUp == EKeys::Escape || IsEscapeKey(KeyUp)))
 		{
 			return FReply::Handled();
 		}
