PR #9106: Bug Fix: ValveIndex_Left_Trackpad_Touch would never call release (Contributed by [REDACTED_PERSON])

diff --git a/Engine/Source/Runtime/InputCore/Private/InputCoreTypes.cpp b/Engine/Source/Runtime/InputCore/Private/InputCoreTypes.cpp
--- a/Engine/Source/Runtime/InputCore/Private/InputCoreTypes.cpp
+++ b/Engine/Source/Runtime/InputCore/Private/InputCoreTypes.cpp
@@ -909,7 +909,7 @@ void EKeys::Initialize()
 	AddKey(FKeyDetails(EKeys::ValveIndex_Left_Trackpad_X, LOCTEXT("ValveIndex_Left_Trackpad_X", "Valve Index (L) Trackpad X"), FKeyDetails::GamepadKey | FKeyDetails::Axis1D | FKeyDetails::NotBlueprintBindableKey, "ValveIndex"));
 	AddKey(FKeyDetails(EKeys::ValveIndex_Left_Trackpad_Y, LOCTEXT("ValveIndex_Left_Trackpad_Y", "Valve Index (L) Trackpad Y"), FKeyDetails::GamepadKey | FKeyDetails::Axis1D | FKeyDetails::NotBlueprintBindableKey, "ValveIndex"));
 	AddKey(FKeyDetails(EKeys::ValveIndex_Left_Trackpad_Force, LOCTEXT("ValveIndex_Left_Trackpad_Force", "Valve Index (L) Trackpad Force"), FKeyDetails::GamepadKey | FKeyDetails::Axis1D | FKeyDetails::NotBlueprintBindableKey, "ValveIndex"));
-	AddKey(FKeyDetails(EKeys::ValveIndex_Left_Trackpad_Touch, LOCTEXT("ValveIndex_Left_Trackpad_Touch", "Valve Index (L) Trackpad Touch"), FKeyDetails::GamepadKey | FKeyDetails::Axis1D | FKeyDetails::NotBlueprintBindableKey, "ValveIndex"));
+	AddKey(FKeyDetails(EKeys::ValveIndex_Left_Trackpad_Touch, LOCTEXT("ValveIndex_Left_Trackpad_Touch", "Valve Index (L) Trackpad Touch"), FKeyDetails::GamepadKey | FKeyDetails::NotBlueprintBindableKey, "ValveIndex"));
 	AddKey(FKeyDetails(EKeys::ValveIndex_Left_Trackpad_Up, LOCTEXT("ValveIndex_Left_Trackpad_Up", "Valve Index (L) Trackpad Up"), FKeyDetails::GamepadKey | FKeyDetails::NotBlueprintBindableKey, "ValveIndex"));
 	AddKey(FKeyDetails(EKeys::ValveIndex_Left_Trackpad_Down, LOCTEXT("ValveIndex_Left_Trackpad_Down", "Valve Index (L) Trackpad Down"), FKeyDetails::GamepadKey | FKeyDetails::NotBlueprintBindableKey, "ValveIndex"));
 	AddKey(FKeyDetails(EKeys::ValveIndex_Left_Trackpad_Left, LOCTEXT("ValveIndex_Left_Trackpad_Left", "Valve Index (L) Trackpad Left"), FKeyDetails::GamepadKey | FKeyDetails::NotBlueprintBindableKey, "ValveIndex"));
