Enhanced Input for OpenXR

This allows us to expose the OpenXR input system through the same assets and APIs as Enhanced Input providing a familiar interface to devs.

- UInputAction is equivalent to XrAction
- UInputMappingContext is equivalent to XrActionSet
- FEnhancedActionKeyMapping is equivalent to XrActionSuggestedBinding

However unlike Enhanced input there are additional constraints:

- All actions sets that can be applied during a session need to be provided to xrAttachSessionActionSets and become immutable for the duration of the session.
- All suggested bindings need to be provided to xrSuggestInteractionProfileBindings and also become immutable for the duration of the session.
- The action set priorities are also immutable, but we're currently in the process of drafting an extension that allows them to be changed at any time.

To accommodate these constraints we ask the developer to provide a UPlayerMappableInputConfig containing every UInputMappingContext that will be used for XR and their default priorities.

diff --git a/Engine/Plugins/Runtime/OpenXR/OpenXR.uplugin b/Engine/Plugins/Runtime/OpenXR/OpenXR.uplugin
--- a/Engine/Plugins/Runtime/OpenXR/OpenXR.uplugin
+++ b/Engine/Plugins/Runtime/OpenXR/OpenXR.uplugin
@@ -1,19 +1,19 @@
 {
-	"FileVersion" : 3,
-	"Version" : 1,
-	"VersionName" : "1.0",
-	"FriendlyName" : "OpenXR",
-	"Description" : "OpenXR is an open VR/AR standard",
-	"Category" : "Virtual Reality",
-	"CreatedBy" : "Epic Games, Inc.",
-	"CreatedByURL" : "http://epicgames.com",
-	"DocsURL" : "",
-	"MarketplaceURL" : "",
-	"SupportURL" : "",
-	"EnabledByDefault" : false,
-	"CanContainContent" : true,
-	"IsBetaVersion" : false,
-	"Installed" : false,
+	"FileVersion": 3,
+	"Version": 1,
+	"VersionName": "1.0",
+	"FriendlyName": "OpenXR",
+	"Description": "OpenXR is an open VR/AR standard",
+	"Category": "Virtual Reality",
+	"CreatedBy": "Epic Games, Inc.",
+	"CreatedByURL": "http://epicgames.com",
+	"DocsURL": "",
+	"MarketplaceURL": "",
+	"SupportURL": "",
+	"EnabledByDefault": false,
+	"CanContainContent": true,
+	"IsBetaVersion": false,
+	"Installed": false,
 	"SupportedTargetPlatforms": [
 		"Win64",
 		"Linux",
@@ -59,5 +59,11 @@
 				"Linux"
 			]
 		}
+	],
+	"Plugins": [
+		{
+			"Name": "EnhancedInput",
+			"Enabled": true
+		}
 	]
 }
diff --git a/Engine/Plugins/Runtime/OpenXR/Source/OpenXREditor/OpenXREditor.Build.cs b/Engine/Plugins/Runtime/OpenXR/Source/OpenXREditor/OpenXREditor.Build.cs
--- a/Engine/Plugins/Runtime/OpenXR/Source/OpenXREditor/OpenXREditor.Build.cs
+++ b/Engine/Plugins/Runtime/OpenXR/Source/OpenXREditor/OpenXREditor.Build.cs
@@ -10,7 +10,10 @@ public class OpenXREditor : ModuleRules
 			new string[] {
 				"Core",
 				"CoreUObject",
+				"EditorFramework",
+				"UnrealEd",
 				"OpenXRHMD",
+				"OpenXRInput"
 			}
 		);
 
diff --git a/Engine/Plugins/Runtime/OpenXR/Source/OpenXREditor/Private/OpenXREditorModule.cpp b/Engine/Plugins/Runtime/OpenXR/Source/OpenXREditor/Private/OpenXREditorModule.cpp
--- a/Engine/Plugins/Runtime/OpenXR/Source/OpenXREditor/Private/OpenXREditorModule.cpp
+++ b/Engine/Plugins/Runtime/OpenXR/Source/OpenXREditor/Private/OpenXREditorModule.cpp
@@ -2,16 +2,41 @@
 
 #include "OpenXREditorModule.h"
 #include "OpenXRAssetDirectory.h"
+#include "OpenXRInputSettings.h"
+
+#include "ISettingsModule.h"
 #include "Modules/ModuleManager.h"
 
+#define LOCTEXT_NAMESPACE "OpenXR"
+
 void FOpenXREditorModule::StartupModule()
 {
 	FOpenXRAssetDirectory::LoadForCook();
+
+	// register settings
+	ISettingsModule* SettingsModule = FModuleManager::GetModulePtr<ISettingsModule>("Settings");
+
+	if (SettingsModule != nullptr)
+	{
+		SettingsModule->RegisterSettings("Project", "Plugins", "OpenXR",
+			LOCTEXT("OpenXRInputSettingsName", "OpenXR Input"),
+			LOCTEXT("OpenXRInputSettingsDescription", "Project settings for OpenXR plugin"),
+			GetMutableDefault<UOpenXRInputSettings>()
+		);
+	}
 }
 
 void FOpenXREditorModule::ShutdownModule()
 {
 	FOpenXRAssetDirectory::ReleaseAll();
+
+	// unregister settings
+	ISettingsModule* SettingsModule = FModuleManager::GetModulePtr<ISettingsModule>("Settings");
+
+	if (SettingsModule != nullptr)
+	{
+		SettingsModule->UnregisterSettings("Project", "Plugins", "OpenXR");
+	}
 }
 
 IMPLEMENT_MODULE(FOpenXREditorModule, OpenXREditor);
diff --git a/Engine/Plugins/Runtime/OpenXR/Source/OpenXRHMD/Public/OpenXRCore.h b/Engine/Plugins/Runtime/OpenXR/Source/OpenXRHMD/Public/OpenXRCore.h
--- a/Engine/Plugins/Runtime/OpenXR/Source/OpenXRHMD/Public/OpenXRCore.h
+++ b/Engine/Plugins/Runtime/OpenXR/Source/OpenXRHMD/Public/OpenXRCore.h
@@ -106,6 +106,11 @@ FORCEINLINE XrRect2Di ToXrRect(FIntRect Rect)
 	return XrRect2Di{ { Rect.Min.X, Rect.Min.Y }, { Rect.Width(), Rect.Height() } };
 }
 
+FORCEINLINE FVector2D ToFVector2D(XrVector2f Vector, float Scale = 1.0f)
+{
+	return FVector2D(Vector.x * Scale, Vector.y * Scale);
+}
+
 FORCEINLINE FVector2D ToFVector2D(XrExtent2Df Extent, float Scale = 1.0f)
 {
 	return FVector2D(Extent.width * Scale, Extent.height * Scale);
diff --git a/Engine/Plugins/Runtime/OpenXR/Source/OpenXRInput/OpenXRInput.Build.cs b/Engine/Plugins/Runtime/OpenXR/Source/OpenXRInput/OpenXRInput.Build.cs
--- a/Engine/Plugins/Runtime/OpenXR/Source/OpenXRInput/OpenXRInput.Build.cs
+++ b/Engine/Plugins/Runtime/OpenXR/Source/OpenXRInput/OpenXRInput.Build.cs
@@ -33,7 +33,9 @@ namespace UnrealBuildTool.Rules
                 }
                 );
 
-            AddEngineThirdPartyPrivateStaticDependencies(Target, "OpenXR");
+			PublicDependencyModuleNames.Add("EnhancedInput");
+
+			AddEngineThirdPartyPrivateStaticDependencies(Target, "OpenXR");
 
             if (Target.bBuildEditor == true)
             {
diff --git a/Engine/Plugins/Runtime/OpenXR/Source/OpenXRInput/Private/OpenXRInput.cpp b/Engine/Plugins/Runtime/OpenXR/Source/OpenXRInput/Private/OpenXRInput.cpp
--- a/Engine/Plugins/Runtime/OpenXR/Source/OpenXRInput/Private/OpenXRInput.cpp
+++ b/Engine/Plugins/Runtime/OpenXR/Source/OpenXRInput/Private/OpenXRInput.cpp
@@ -1,6 +1,7 @@
 // Copyright Epic Games, Inc. All Rights Reserved.
 
 #include "OpenXRInput.h"
+#include "OpenXRInputSettings.h"
 #include "OpenXRHMD.h"
 #include "OpenXRCore.h"
 #include "UObject/UObjectIterator.h"
@@ -8,6 +9,11 @@
 #include "IOpenXRExtensionPlugin.h"
 #include "HeadMountedDisplayFunctionLibrary.h"
 
+#include "EnhancedInputModule.h"
+#include "InputAction.h"
+#include "InputMappingContext.h"
+#include "PlayerMappableInputConfig.h"
+
 #if WITH_EDITOR
 #include "Editor/EditorEngine.h"
 #include "Editor.h"
@@ -43,6 +49,8 @@ FORCEINLINE XrPath GetPath(XrInstance Instance, const FString& PathString)
 
 FORCEINLINE void FilterActionName(const char* InActionName, char* OutActionName)
 {
+	static_assert(XR_MAX_ACTION_NAME_SIZE == XR_MAX_ACTION_SET_NAME_SIZE);
+
 	// Ensure the action name is a well-formed path
 	size_t i;
 	for (i = 0; i < XR_MAX_ACTION_NAME_SIZE - 1 && InActionName[i] != '\0'; i++)
@@ -53,6 +61,22 @@ FORCEINLINE void FilterActionName(const char* InActionName, char* OutActionName)
 	OutActionName[i] = '\0';
 }
 
+FORCEINLINE XrActionType ToActionType(EInputActionValueType ValueType)
+{
+	switch (ValueType)
+	{
+		case EInputActionValueType::Boolean: return XR_ACTION_TYPE_BOOLEAN_INPUT;
+		case EInputActionValueType::Axis1D: return XR_ACTION_TYPE_FLOAT_INPUT;
+		case EInputActionValueType::Axis2D: return XR_ACTION_TYPE_VECTOR2F_INPUT;
+		case EInputActionValueType::Axis3D:
+			// TODO: Add 3D vector support to OpenXR Input
+			ensure(false);
+			return XR_ACTION_TYPE_VECTOR2F_INPUT;
+	}
+	ensure(false);
+	return (XrActionType)0;
+}
+
 TSharedPtr< class IInputDevice > FOpenXRInputPlugin::CreateInputDevice(const TSharedRef< FGenericApplicationMessageHandler >& InMessageHandler)
 {
 	if (InputDevice)
@@ -91,11 +115,22 @@ void FOpenXRInputPlugin::StartupModule()
 	InputDevice = MakeShared<FOpenXRInput>(OpenXRHMD);
 }
 
-FOpenXRInputPlugin::FOpenXRAction::FOpenXRAction(XrActionSet InActionSet, XrActionType InActionType, const FName& InName, const TArray<XrPath>& SubactionPaths)
+FOpenXRInputPlugin::FOpenXRAction::FOpenXRAction(XrActionSet InActionSet,
+	XrActionType InActionType, const FName& InName, const FString& InLocalizedName,
+	const TArray<XrPath>& InSubactionPaths, const TObjectPtr<const UInputAction>& InObject)
+	: FOpenXRAction(InActionSet, InActionType, InName, InLocalizedName, InSubactionPaths)
+{
+	Object = InObject;
+}
+
+FOpenXRInputPlugin::FOpenXRAction::FOpenXRAction(XrActionSet InActionSet,
+	XrActionType InActionType, const FName& InName, const FString& InLocalizedName,
+	const TArray<XrPath>& InSubactionPaths)
 	: Set(InActionSet)
 	, Type(InActionType)
 	, Name(InName)
 	, Handle(XR_NULL_HANDLE)
+	, Object()
 	, KeyMap()
 {
 	char ActionName[NAME_SIZE];
@@ -106,13 +141,55 @@ FOpenXRInputPlugin::FOpenXRAction::FOpenXRAction(XrActionSet InActionSet, XrActi
 	Info.next = nullptr;
 	FilterActionName(ActionName, Info.actionName);
 	Info.actionType = Type;
-	Info.countSubactionPaths = SubactionPaths.Num();
-	Info.subactionPaths = SubactionPaths.GetData();
-	FCStringAnsi::Strcpy(Info.localizedActionName, XR_MAX_LOCALIZED_ACTION_NAME_SIZE, ActionName);
+	Info.countSubactionPaths = InSubactionPaths.Num();
+	Info.subactionPaths = InSubactionPaths.GetData();
+	if (!InLocalizedName.IsEmpty())
+	{
+		FTCHARToUTF8_Convert::Convert(Info.localizedActionName, XR_MAX_LOCALIZED_ACTION_NAME_SIZE, *InLocalizedName, InLocalizedName.Len() + 1);
+	}
+	else
+	{
+		FCStringAnsi::Strcpy(Info.localizedActionName, XR_MAX_LOCALIZED_ACTION_NAME_SIZE, ActionName);
+	}
 
 	XR_ENSURE(xrCreateAction(Set, &Info, &Handle));
 }
 
+FOpenXRInputPlugin::FOpenXRActionSet::FOpenXRActionSet(XrInstance InInstance,
+	const FName& InName, const FString& InLocalizedName, uint32 InPriority,
+	const TObjectPtr<const UInputMappingContext>& InObject)
+	: FOpenXRActionSet(InInstance, InName, InLocalizedName, InPriority)
+{
+	Object = InObject;
+}
+
+FOpenXRInputPlugin::FOpenXRActionSet::FOpenXRActionSet(XrInstance InInstance,
+	const FName& InName, const FString& InLocalizedName, uint32 InPriority)
+	: Handle(XR_NULL_HANDLE)
+	, Name(InName)
+	, LocalizedName(InLocalizedName)
+	, Object()
+{
+	char ActionName[NAME_SIZE];
+	Name.GetPlainANSIString(ActionName);
+
+	XrActionSet ActionSet = XR_NULL_HANDLE;
+	XrActionSetCreateInfo Info;
+	Info.type = XR_TYPE_ACTION_SET_CREATE_INFO;
+	Info.next = nullptr;
+	FilterActionName(ActionName, Info.actionSetName);
+	if (!InLocalizedName.IsEmpty())
+	{
+		FTCHARToUTF8_Convert::Convert(Info.localizedActionSetName, XR_MAX_LOCALIZED_ACTION_SET_NAME_SIZE, *InLocalizedName, InLocalizedName.Len() + 1);
+	}
+	else
+	{
+		FCStringAnsi::Strcpy(Info.localizedActionSetName, XR_MAX_LOCALIZED_ACTION_SET_NAME_SIZE, ActionName);
+	}
+	Info.priority = InPriority;
+	XR_ENSURE(xrCreateActionSet(InInstance, &Info, &Handle));
+}
+
 FOpenXRInputPlugin::FOpenXRController::FOpenXRController(XrActionSet InActionSet, XrPath InUserPath, const char* InName)
 	: ActionSet(InActionSet)
 	, UserPath(InUserPath)
@@ -165,9 +242,17 @@ FOpenXRInputPlugin::FInteractionProfile::FInteractionProfile(XrPath InProfile, b
 
 FOpenXRInputPlugin::FOpenXRInput::FOpenXRInput(FOpenXRHMD* HMD)
 	: OpenXRHMD(HMD)
+	, Profiles()
 	, ActionSets()
-	, Actions()
+	, PluginActionSets()
+	, SubactionPaths()
+	, LegacyActions()
+	, EnhancedActions()
 	, Controllers()
+	, MotionSourceToControllerHandMap()
+	, ControllerSet()
+	, LegacySet()
+	, MappableInputConfig(nullptr)
 	, bActionsBound(false)
 	, bDirectionalBindingSupported(false)
 	, MessageHandler(new FGenericApplicationMessageHandler())
@@ -256,28 +341,8 @@ void FOpenXRInputPlugin::FOpenXRInput::BuildActions()
 	check(Instance);
 	DestroyActions();
 
-	XrActionSet ActionSet = XR_NULL_HANDLE;
-	XrActionSetCreateInfo SetInfo;
-	SetInfo.type = XR_TYPE_ACTION_SET_CREATE_INFO;
-	SetInfo.next = nullptr;
-	FCStringAnsi::Strcpy(SetInfo.actionSetName, XR_MAX_ACTION_SET_NAME_SIZE, "ue");
-	FCStringAnsi::Strcpy(SetInfo.localizedActionSetName, XR_MAX_LOCALIZED_ACTION_SET_NAME_SIZE, "Unreal Engine");
-	SetInfo.priority = 0;
-	XR_ENSURE(xrCreateActionSet(Instance, &SetInfo, &ActionSet));
-
-	XrPath LeftHand = GetPath(Instance, "/user/hand/left");
-	XrPath RightHand = GetPath(Instance, "/user/hand/right");
-	XrPath HMD = GetPath(Instance, "/user/head");
-
-	// Controller poses
-	OpenXRHMD->ResetActionDevices();
-	Controllers.Add(EControllerHand::Left, FOpenXRController(ActionSet, LeftHand, "Left Controller"));
-	Controllers.Add(EControllerHand::Right, FOpenXRController(ActionSet, RightHand, "Right Controller"));
-	Controllers.Add(EControllerHand::HMD, FOpenXRController(ActionSet, HMD, "HMD"));
-
 	// Generate a map of all supported interaction profiles
 	XrPath SimpleControllerPath = GetPath(Instance, "/interaction_profiles/khr/simple_controller");
-	TMap<FString, FInteractionProfile> Profiles;
 	Profiles.Add("SimpleController", FInteractionProfile(SimpleControllerPath, true));
 	Profiles.Add("Vive", FInteractionProfile(GetPath(Instance, "/interaction_profiles/htc/vive_controller"), true));
 	Profiles.Add("MixedReality", FInteractionProfile(GetPath(Instance, "/interaction_profiles/microsoft/motion_controller"), true));
@@ -297,10 +362,93 @@ void FOpenXRInputPlugin::FOpenXRInput::BuildActions()
 		}
 	}
 
+	// Create an engine action set for pose input and haptic output
+	FOpenXRActionSet ActionSet(Instance, "controllers", "Controllers", 0);
+	ControllerSet.actionSet = ActionSet.Handle;
+	ControllerSet.subactionPath = XR_NULL_PATH;
+
+	XrPath LeftHand = GetPath(Instance, "/user/hand/left");
+	XrPath RightHand = GetPath(Instance, "/user/hand/right");
+	XrPath HMD = GetPath(Instance, "/user/head");
+
+	// Controller poses
+	OpenXRHMD->ResetActionDevices();
+	Controllers.Add(EControllerHand::Left, FOpenXRController(ActionSet.Handle, LeftHand, "Left Controller"));
+	Controllers.Add(EControllerHand::Right, FOpenXRController(ActionSet.Handle, RightHand, "Right Controller"));
+	Controllers.Add(EControllerHand::HMD, FOpenXRController(ActionSet.Handle, HMD, "HMD"));
+	ActionSets.Emplace(MoveTemp(ActionSet));
+
 	// Generate a list of the sub-action paths so we can query the left/right hand individually
 	SubactionPaths.Add(LeftHand);
 	SubactionPaths.Add(RightHand);
 
+	UOpenXRInputSettings* InputSettings = GetMutableDefault<UOpenXRInputSettings>();
+	if (InputSettings && InputSettings->bUseEnhancedInput)
+	{
+		BuildEnhancedActions(InputSettings->MappableInputConfig);
+	}
+	else
+	{
+		BuildLegacyActions();
+	}
+
+	for (TPair<FString, FInteractionProfile>& Pair : Profiles)
+	{
+		FInteractionProfile& Profile = Pair.Value;
+
+		// Only suggest interaction profile bindings if the developer has provided bindings for them
+		// An exception is made for the Simple Controller Profile which is always bound as a fallback
+		if (Profile.Bindings.Num() > 0 || Profile.Path == SimpleControllerPath)
+		{
+			// Add the bindings for the controller pose and haptics
+			Profile.Bindings.Add(XrActionSuggestedBinding{
+				Controllers[EControllerHand::Left].GripAction, GetPath(Instance, "/user/hand/left/input/grip/pose")
+				});
+			Profile.Bindings.Add(XrActionSuggestedBinding{
+				Controllers[EControllerHand::Right].GripAction, GetPath(Instance, "/user/hand/right/input/grip/pose")
+				});
+			Profile.Bindings.Add(XrActionSuggestedBinding{
+				Controllers[EControllerHand::Left].AimAction, GetPath(Instance, "/user/hand/left/input/aim/pose")
+				});
+			Profile.Bindings.Add(XrActionSuggestedBinding{
+				Controllers[EControllerHand::Right].AimAction, GetPath(Instance, "/user/hand/right/input/aim/pose")
+				});
+
+			if (Profile.HasHaptics)
+			{
+				Profile.Bindings.Add(XrActionSuggestedBinding{
+					Controllers[EControllerHand::Left].VibrationAction, GetPath(Instance, "/user/hand/left/output/haptic")
+					});
+				Profile.Bindings.Add(XrActionSuggestedBinding{
+					Controllers[EControllerHand::Right].VibrationAction, GetPath(Instance, "/user/hand/right/output/haptic")
+					});
+			}
+
+			XrInteractionProfileSuggestedBinding InteractionProfile;
+			InteractionProfile.type = XR_TYPE_INTERACTION_PROFILE_SUGGESTED_BINDING;
+			InteractionProfile.next = nullptr;
+			InteractionProfile.interactionProfile = Profile.Path;
+			InteractionProfile.countSuggestedBindings = Profile.Bindings.Num();
+			InteractionProfile.suggestedBindings = Profile.Bindings.GetData();
+
+			XR_ENSURE(xrSuggestInteractionProfileBindings(Instance, &InteractionProfile));
+		}
+	}
+
+	Controllers[EControllerHand::Left].AddActionDevices(OpenXRHMD);
+	Controllers[EControllerHand::Right].AddActionDevices(OpenXRHMD);
+}
+
+void FOpenXRInputPlugin::FOpenXRInput::BuildLegacyActions()
+{
+	XrInstance Instance = OpenXRHMD->GetInstance();
+	check(Instance);
+
+	// Create an engine action set for pose input and haptic output
+	FOpenXRActionSet ActionSet(Instance, "ue", "Unreal Engine", 0);
+	LegacySet.actionSet = ActionSet.Handle;
+	LegacySet.subactionPath = XR_NULL_PATH;
+
 	auto InputSettings = GetMutableDefault<UInputSettings>();
 	if (InputSettings != nullptr)
 	{
@@ -308,14 +456,14 @@ void FOpenXRInputPlugin::FOpenXRInput::BuildActions()
 		InputSettings->GetActionNames(ActionNames);
 		for (const FName& ActionName : ActionNames)
 		{
-			FOpenXRAction Action(ActionSet, XR_ACTION_TYPE_BOOLEAN_INPUT, ActionName, SubactionPaths);
+			FOpenXRAction Action(ActionSet.Handle, XR_ACTION_TYPE_BOOLEAN_INPUT, ActionName, ActionName.ToString(), SubactionPaths);
 			TArray<FInputActionKeyMapping> Mappings;
 			InputSettings->GetActionMappingByName(ActionName, Mappings);
 
 			// If the developer didn't suggest any XR bindings for this action we won't expose it to the runtime
-			if (SuggestBindings(Instance, Action, Mappings, Profiles) > 0)
+			if (SuggestBindings(Instance, Action, Mappings) > 0)
 			{
-				Actions.Add(Action);
+				LegacyActions.Add(Action);
 			}
 			else
 			{
@@ -327,14 +475,14 @@ void FOpenXRInputPlugin::FOpenXRInput::BuildActions()
 		InputSettings->GetAxisNames(AxisNames);
 		for (const FName& AxisName : AxisNames)
 		{
-			FOpenXRAction Action(ActionSet, XR_ACTION_TYPE_FLOAT_INPUT, AxisName, SubactionPaths);
+			FOpenXRAction Action(ActionSet.Handle, XR_ACTION_TYPE_FLOAT_INPUT, AxisName, AxisName.ToString(), SubactionPaths);
 			TArray<FInputAxisKeyMapping> Mappings;
 			InputSettings->GetAxisMappingByName(AxisName, Mappings);
 
 			// If the developer didn't suggest any XR bindings for this action we won't expose it to the runtime
-			if (SuggestBindings(Instance, Action, Mappings, Profiles) > 0)
+			if (SuggestBindings(Instance, Action, Mappings) > 0)
 			{
-				Actions.Add(Action);
+				LegacyActions.Add(Action);
 			}
 			else
 			{
@@ -353,8 +501,8 @@ void FOpenXRInputPlugin::FOpenXRInput::BuildActions()
 			[this, &ActionSet](XrActionType InActionType, const FName& InName, const TArray<XrPath>& InSubactionPaths)
 			{
 				// TODO?: Log deprecation warning at runtime, since overridden deprecated interface methods don't warn at compile time?
-				FOpenXRAction Action(ActionSet, InActionType, InName, InSubactionPaths);
-				Actions.Add(Action);
+				FOpenXRAction Action(ActionSet.Handle, InActionType, InName, InName.ToString(), InSubactionPaths);
+				LegacyActions.Add(Action);
 				return Action.Handle;
 			}
 		);
@@ -367,125 +515,92 @@ void FOpenXRInputPlugin::FOpenXRInput::BuildActions()
 			[this, Instance]
 			(const IOpenXRExtensionPlugin::FActionSetParams& Params) -> XrActionSet
 			{
-				XrActionSet OutActionSet = XR_NULL_HANDLE;
-
-				char ActionSetName[NAME_SIZE];
-				Params.Name.GetPlainANSIString(ActionSetName);
-
-				XrActionSetCreateInfo SetInfo;
-				SetInfo.type = XR_TYPE_ACTION_SET_CREATE_INFO;
-				SetInfo.next = nullptr;
-				FCStringAnsi::Strcpy(SetInfo.actionSetName, XR_MAX_ACTION_SET_NAME_SIZE, ActionSetName);
-				SetInfo.priority = Params.Priority;
-
-				if (!Params.LocalizedName.IsEmpty())
-				{
-					const FString& LocalizedNameStr = Params.LocalizedName.ToString();
-					FTCHARToUTF8_Convert::Convert(SetInfo.localizedActionSetName, XR_MAX_LOCALIZED_ACTION_SET_NAME_SIZE, *LocalizedNameStr, LocalizedNameStr.Len() + 1);
-				}
-				else
-				{
-					FCStringAnsi::Strcpy(SetInfo.localizedActionSetName, XR_MAX_LOCALIZED_ACTION_SET_NAME_SIZE, ActionSetName);
-				}
-
-
-				if (XR_ENSURE(xrCreateActionSet(Instance, &SetInfo, &OutActionSet)))
-				{
-					PluginActionSets.Add(OutActionSet);
-					return OutActionSet;
-				}
-				else
-				{
-					return XR_NULL_HANDLE;
-				}
+				FOpenXRActionSet ActionSet(Instance, Params.Name, Params.LocalizedName.ToString(), Params.Priority);
+				const XrActionSet ReturnHandle = ActionSet.Handle;
+				PluginActionSets.Emplace(MoveTemp(ActionSet));
+				return ReturnHandle;
 			};
 
 		IOpenXRExtensionPlugin::FCreateActionFunc CreateActionFunc =
-			[this, Instance, ActionSet, &Profiles]
+			[this, Instance, ActionSet]
 			(const IOpenXRExtensionPlugin::FActionParams& Params) -> XrAction
 			{
 				// TODO: Use Params.LocalizedName
-				const XrActionSet UseActionSet = Params.Set == XR_NULL_HANDLE ? ActionSet : Params.Set;
-				FOpenXRAction Action(UseActionSet, Params.Type, Params.Name, Params.SubactionPaths);
+				const XrActionSet UseActionSet = Params.Set == XR_NULL_HANDLE ? ActionSet.Handle : Params.Set;
+				FOpenXRAction Action(UseActionSet, Params.Type, Params.Name, Params.LocalizedName.ToString(), Params.SubactionPaths);
 				const XrAction ReturnHandle = Action.Handle;
 				for (const FKey& Key : Params.SuggestedBindings)
 				{
-					SuggestBindingForKey(Instance, Action, Key, Profiles);
+					SuggestBindingForKey(Instance, Action, Key);
 				}
-				Actions.Emplace(MoveTemp(Action));
+				LegacyActions.Emplace(MoveTemp(Action));
 				return ReturnHandle;
 			};
 
 		Plugin->AddActions(Instance, CreateActionSetFunc, CreateActionFunc);
 	}
 
-	for (TPair<FString, FInteractionProfile>& Pair : Profiles)
+	ActionSets.Emplace(MoveTemp(ActionSet));
+}
+
+void FOpenXRInputPlugin::FOpenXRInput::BuildEnhancedActions(const FSoftObjectPath& MappableInputConfigPath)
+{
+	XrInstance Instance = OpenXRHMD->GetInstance();
+	check(Instance);
+
+	MappableInputConfig = (UPlayerMappableInputConfig*)MappableInputConfigPath.TryLoad();
+	if (!MappableInputConfig)
 	{
-		FInteractionProfile& Profile = Pair.Value;
+		return;
+	}
+	MappableInputConfig->AddToRoot();
 
-		// Only suggest interaction profile bindings if the developer has provided bindings for them
-		// An exception is made for the Simple Controller Profile which is always bound as a fallback
-		if (Profile.Bindings.Num() > 0 || Profile.Path == SimpleControllerPath)
+	for (const TPair<TObjectPtr<UInputMappingContext>, int32> MappingContext : MappableInputConfig->GetMappingContexts())
+	{
+		FOpenXRActionSet ActionSet(Instance, MappingContext.Key->GetFName(), MappingContext.Key->ContextDescription.ToString(), MappingContext.Value, MappingContext.Key);
+		for (const FEnhancedActionKeyMapping& Mapping : MappingContext.Key->GetMappings())
 		{
-			// Add the bindings for the controller pose and haptics
-			Profile.Bindings.Add(XrActionSuggestedBinding {
-				Controllers[EControllerHand::Left].GripAction, GetPath(Instance, "/user/hand/left/input/grip/pose")
-				});
-			Profile.Bindings.Add(XrActionSuggestedBinding {
-				Controllers[EControllerHand::Right].GripAction, GetPath(Instance, "/user/hand/right/input/grip/pose")
-				});
-			Profile.Bindings.Add(XrActionSuggestedBinding{
-				Controllers[EControllerHand::Left].AimAction, GetPath(Instance, "/user/hand/left/input/aim/pose")
-				});
-			Profile.Bindings.Add(XrActionSuggestedBinding{
-				Controllers[EControllerHand::Right].AimAction, GetPath(Instance, "/user/hand/right/input/aim/pose")
-				});
-
-			if (Profile.HasHaptics)
+			// Try to find an existing action
+			FName ActionName = Mapping.Action->GetFName();
+			FOpenXRAction* Action = EnhancedActions.FindByPredicate([ActionName](const FOpenXRAction& Action) { return Action.Name == ActionName; });
+			if (!Action)
 			{
-				Profile.Bindings.Add(XrActionSuggestedBinding{
-					Controllers[EControllerHand::Left].VibrationAction, GetPath(Instance, "/user/hand/left/output/haptic")
-					});
-				Profile.Bindings.Add(XrActionSuggestedBinding{
-					Controllers[EControllerHand::Right].VibrationAction, GetPath(Instance, "/user/hand/right/output/haptic")
-					});
-			}
+				// No action found, create a new one
+				FString LocalizedName = Mapping.Action->ActionDescription.ToString();
+				XrActionType ActionType = ToActionType(Mapping.Action->ValueType);
+				if (!ActionType)
+				{
+					continue;
+				}
 
-			XrInteractionProfileSuggestedBinding InteractionProfile;
-			InteractionProfile.type = XR_TYPE_INTERACTION_PROFILE_SUGGESTED_BINDING;
-			InteractionProfile.next = nullptr;
-			InteractionProfile.interactionProfile = Profile.Path;
-			InteractionProfile.countSuggestedBindings = Profile.Bindings.Num();
-			InteractionProfile.suggestedBindings = Profile.Bindings.GetData();
+				FOpenXRAction NewAction(ActionSet.Handle, ActionType, ActionName, LocalizedName, SubactionPaths, Mapping.Action);
+				int32 Index = EnhancedActions.Emplace(MoveTemp(NewAction));
+				Action = &EnhancedActions[Index];
+			}
 
-			XR_ENSURE(xrSuggestInteractionProfileBindings(Instance, &InteractionProfile));
+			SuggestBindingForKey(Instance, *Action, Mapping.Key, Mapping.Modifiers, Mapping.Triggers);
 		}
-	}
-
-	Controllers[EControllerHand::Left].AddActionDevices(OpenXRHMD);
-	Controllers[EControllerHand::Right].AddActionDevices(OpenXRHMD);
-
-	// Add an active set for each sub-action path so we can use the subaction paths later
-	// TODO: Runtimes already allow us to do the same by simply specifying "subactionPath"
-	// as XR_NULL_PATH, we're just being verbose for safety.
-	//for (XrPath Subaction : SubactionPaths)
-	{
-		XrActiveActionSet ActiveSet;
-		ActiveSet.actionSet = ActionSet;
-		ActiveSet.subactionPath = XR_NULL_PATH;
-		ActionSets.Add(ActiveSet);
+		ActionSets.Emplace(MoveTemp(ActionSet));
 	}
 }
 
 void FOpenXRInputPlugin::FOpenXRInput::DestroyActions()
 {
 	// Destroying an action set will also destroy all actions in the set
-	for (XrActiveActionSet& ActionSet : ActionSets)
+	for (const FOpenXRActionSet& ActionSet : ActionSets)
 	{
-		xrDestroyActionSet(ActionSet.actionSet);
+		xrDestroyActionSet(ActionSet.Handle);
 	}
 
-	Actions.Reset();
+	// Clear the active controller and legacy action set
+	ControllerSet.actionSet = XR_NULL_HANDLE;
+	ControllerSet.subactionPath = XR_NULL_PATH;
+	LegacySet.actionSet = XR_NULL_HANDLE;
+	LegacySet.subactionPath = XR_NULL_PATH;
+
+	Profiles.Reset();
+	LegacyActions.Reset();
+	EnhancedActions.Reset();
 	Controllers.Reset();
 	SubactionPaths.Reset();
 	ActionSets.Reset();
@@ -493,14 +608,14 @@ void FOpenXRInputPlugin::FOpenXRInput::DestroyActions()
 }
 
 template<typename T>
-int32 FOpenXRInputPlugin::FOpenXRInput::SuggestBindings(XrInstance Instance, FOpenXRAction& Action, const TArray<T>& Mappings, TMap<FString, FInteractionProfile>& Profiles)
+int32 FOpenXRInputPlugin::FOpenXRInput::SuggestBindings(XrInstance Instance, FOpenXRAction& Action, const TArray<T>& Mappings)
 {
 	int32 SuggestedBindings = 0;
 
 	// Add suggested bindings for every mapping
 	for (const T& InputKey : Mappings)
 	{
-		if (SuggestBindingForKey(Instance, Action, InputKey.Key, Profiles))
+		if (SuggestBindingForKey(Instance, Action, InputKey.Key))
 		{
 			++SuggestedBindings;
 		}
@@ -509,7 +624,7 @@ int32 FOpenXRInputPlugin::FOpenXRInput::SuggestBindings(XrInstance Instance, FOp
 	return SuggestedBindings;
 }
 
-bool FOpenXRInputPlugin::FOpenXRInput::SuggestBindingForKey(XrInstance Instance, FOpenXRAction& Action, const FKey& InFKey, TMap<FString, FInteractionProfile>& Profiles)
+bool FOpenXRInputPlugin::FOpenXRInput::SuggestBindingForKey(XrInstance Instance, FOpenXRAction& Action, const FKey& InFKey, const TArray<UInputModifier*>& Modifiers, const TArray<UInputTrigger*>& Triggers)
 {
 	// Key names that are parseable into an OpenXR path have exactly 4 tokens
 	TArray<FString> Tokens;
@@ -534,6 +649,17 @@ bool FOpenXRInputPlugin::FOpenXRInput::SuggestBindingForKey(XrInstance Instance,
 	TPair<XrPath, XrPath> Key(Profile->Path, TopLevel);
 	Action.KeyMap.Add(Key, InFKey.GetFName());
 
+	for (UInputTrigger* Trigger : Triggers)
+	{
+		TObjectPtr<UInputTrigger> Ptr = Trigger;
+		Action.Triggers.AddUnique(Key, Ptr);
+	}
+	for (UInputModifier* Modifier : Modifiers)
+	{
+		TObjectPtr<UInputModifier> Ptr = Modifier;
+		Action.Modifiers.AddUnique(Key, Ptr);
+	}
+
 	// Add the input we want to query with grip being defined as "squeeze" in OpenXR
 	FString Identifier = Tokens[2].ToLower();
 	if (Identifier == "grip")
@@ -583,9 +709,10 @@ bool FOpenXRInputPlugin::FOpenXRInput::SuggestBindingForKey(XrInstance Instance,
 		}
 		Path += "/dpad_" + Component;
 	}
-	else
+	else if (Component != "2d")
 	{
 		// Anything we don't need to translate can pass through
+		// Except for 2D vectors, which don't need a component path
 		Path += "/" + Component;
 	}
 
@@ -594,6 +721,13 @@ bool FOpenXRInputPlugin::FOpenXRInput::SuggestBindingForKey(XrInstance Instance,
 	return true;
 }
 
+bool FOpenXRInputPlugin::FOpenXRInput::SuggestBindingForKey(XrInstance Instance, FOpenXRAction& Action, const FKey& InFKey)
+{
+	TArray<UInputModifier*> Modifiers;
+	TArray<UInputTrigger*> Triggers;
+	return SuggestBindingForKey(Instance, Action, InFKey, Modifiers, Triggers);
+}
+
 void FOpenXRInputPlugin::FOpenXRInput::Tick(float DeltaTime)
 {
 	if (OpenXRHMD == nullptr)
@@ -607,10 +741,12 @@ void FOpenXRInputPlugin::FOpenXRInput::Tick(float DeltaTime)
 	{
 		if (!bActionsBound)
 		{
+			BuildActions();
+
 			// Bind engine action sets
 			TArray<XrActionSet> BindActionSets;
 			for (auto && BindActionSet : ActionSets)
-				BindActionSets.Add(BindActionSet.actionSet);
+				BindActionSets.Add(BindActionSet.Handle);
 
 			// Bind plugin action sets exposed via deprecated method
 			for (IOpenXRExtensionPlugin* Plugin : OpenXRHMD->GetExtensionPlugins())
@@ -627,7 +763,7 @@ void FOpenXRInputPlugin::FOpenXRInput::Tick(float DeltaTime)
 			}
 
 			for (auto&& BindActionSet : PluginActionSets)
-				BindActionSets.Add(BindActionSet);
+				BindActionSets.Add(BindActionSet.Handle);
 
 			XrSessionActionSetsAttachInfo SessionActionSetsAttachInfo;
 			SessionActionSetsAttachInfo.type = XR_TYPE_SESSION_ACTION_SETS_ATTACH_INFO;
@@ -643,6 +779,13 @@ void FOpenXRInputPlugin::FOpenXRInput::Tick(float DeltaTime)
 	{
 		// If the session shut down, clean up.
 		bActionsBound = false;
+
+		// Remove the input config from the root
+		if (MappableInputConfig)
+		{
+			MappableInputConfig->RemoveFromRoot();
+			MappableInputConfig = nullptr;
+		}
 	}
 
 	if (OpenXRHMD->IsFocused())
@@ -664,12 +807,38 @@ void FOpenXRInputPlugin::FOpenXRInput::Tick(float DeltaTime)
 			SyncInfo.next = Plugin->OnSyncActions(Session, SyncInfo.next);
 		}
 
-		ActiveSets.Append(ActionSets);
+		IEnhancedInputModule::Get().GetLibrary()->ForEachSubsystem([this, &ActiveSets](IEnhancedInputSubsystemInterface* Subsystem)
+			{
+				if (Subsystem)
+				{
+					for (const FOpenXRActionSet& ActionSet : ActionSets)
+					{
+						if (ActionSet.Object && Subsystem->HasMappingContext(ActionSet.Object))
+						{
+							XrActiveActionSet ActiveSet;
+							ActiveSet.actionSet = ActionSet.Handle;
+							ActiveSet.subactionPath = XR_NULL_PATH;
+							ActiveSets.Add(ActiveSet);
+						}
+					}
+				}
+			});
+
+		// The controller and legacy action sets are always active
+		if (ControllerSet.actionSet)
+		{
+			ActiveSets.Add(ControllerSet);
+		}
+		if (LegacySet.actionSet)
+		{
+			ActiveSets.Add(LegacySet);
+		}
+
 		SyncInfo.countActiveActionSets = ActiveSets.Num();
 		SyncInfo.activeActionSets = ActiveSets.GetData();
 
 		XR_ENSURE(xrSyncActions(Session, &SyncInfo));
-	
+
 		for (IOpenXRExtensionPlugin* Plugin : OpenXRHMD->GetExtensionPlugins())
 		{
 			Plugin->PostSyncActions(Session);
@@ -713,7 +882,7 @@ void FOpenXRInputPlugin::FOpenXRInput::SendControllerEvents()
 
 		TSet<FName> ActivatedKeys;
 		TPair<XrPath, XrPath> Key(Profile.interactionProfile, Subaction);
-		for (FOpenXRAction& Action : Actions)
+		for (FOpenXRAction& Action : LegacyActions)
 		{
 			XrActionStateGetInfo GetInfo;
 			GetInfo.type = XR_TYPE_ACTION_STATE_GET_INFO;
@@ -787,15 +956,91 @@ void FOpenXRInputPlugin::FOpenXRInput::SendControllerEvents()
 				break;
 			}
 		}
+
+		for (FOpenXRAction& Action : EnhancedActions)
+		{
+			const UInputAction* InputAction = Action.Object;
+
+			XrActionStateGetInfo GetInfo;
+			GetInfo.type = XR_TYPE_ACTION_STATE_GET_INFO;
+			GetInfo.next = nullptr;
+			GetInfo.subactionPath = Subaction;
+			GetInfo.action = Action.Handle;
+
+			FInputActionValue InputValue;
+			switch (Action.Type)
+			{
+			case XR_ACTION_TYPE_BOOLEAN_INPUT:
+			{
+				XrActionStateBoolean State;
+				State.type = XR_TYPE_ACTION_STATE_BOOLEAN;
+				State.next = nullptr;
+				XrResult Result = xrGetActionStateBoolean(Session, &GetInfo, &State);
+				if (XR_SUCCEEDED(Result))
+				{
+					InputValue = FInputActionValue(State.isActive ? (bool)State.currentState : false);
+				}
+				else
+				{
+					continue;
+				}
+			}
+			break;
+			case XR_ACTION_TYPE_FLOAT_INPUT:
+			{
+				XrActionStateFloat State;
+				State.type = XR_TYPE_ACTION_STATE_FLOAT;
+				State.next = nullptr;
+				XrResult Result = xrGetActionStateFloat(Session, &GetInfo, &State);
+				if (XR_SUCCEEDED(Result))
+				{
+					InputValue = FInputActionValue(State.isActive ? State.currentState : 0.0f);
+				}
+				else
+				{
+					continue;
+				}
+			}
+			break;
+			case XR_ACTION_TYPE_VECTOR2F_INPUT:
+			{
+				XrActionStateVector2f State;
+				State.type = XR_TYPE_ACTION_STATE_VECTOR2F;
+				State.next = nullptr;
+				XrResult Result = xrGetActionStateVector2f(Session, &GetInfo, &State);
+				if (XR_SUCCEEDED(Result))
+				{
+					InputValue = FInputActionValue(State.isActive ? ToFVector2D(State.currentState) : FVector2D::ZeroVector);
+				}
+				else
+				{
+					continue;
+				}
+			}
+			break;
+			default:
+				// Other action types are currently unsupported.
+				continue;
+			}
+
+			TArray<TObjectPtr<UInputTrigger>> Triggers;
+			Action.Triggers.MultiFind(Key, Triggers, false);
+			TArray<TObjectPtr<UInputModifier>> Modifiers;
+			Action.Modifiers.MultiFind(Key, Modifiers, false);
+			IEnhancedInputModule::Get().GetLibrary()->ForEachSubsystem([InputAction, InputValue, Triggers, Modifiers](IEnhancedInputSubsystemInterface* Subsystem)
+				{
+					if (Subsystem)
+					{
+						Subsystem->InjectInputForAction(InputAction, InputValue, Modifiers, Triggers);
+					}
+				});
+		}
 	}
 }
 
 void FOpenXRInputPlugin::FOpenXRInput::SetMessageHandler(const TSharedRef< FGenericApplicationMessageHandler >& InMessageHandler)
 {
 	MessageHandler = InMessageHandler;
-#if WITH_EDITOR
-	FEditorDelegates::OnActionAxisMappingsChanged.AddSP(this, &FOpenXRInputPlugin::FOpenXRInput::BuildActions);
-#endif
 }
 
 bool FOpenXRInputPlugin::FOpenXRInput::Exec(UWorld* InWorld, const TCHAR* Cmd, FOutputDevice& Ar)
diff --git a/Engine/Plugins/Runtime/OpenXR/Source/OpenXRInput/Private/OpenXRInput.h b/Engine/Plugins/Runtime/OpenXR/Source/OpenXRInput/Private/OpenXRInput.h
--- a/Engine/Plugins/Runtime/OpenXR/Source/OpenXRInput/Private/OpenXRInput.h
+++ b/Engine/Plugins/Runtime/OpenXR/Source/OpenXRInput/Private/OpenXRInput.h
@@ -9,6 +9,11 @@
 #include "IHapticDevice.h"
 
 class FOpenXRHMD;
+class UInputAction;
+class UInputTrigger;
+class UInputModifier;
+class UInputMappingContext;
+class UPlayerMappableInputConfig;
 struct FInputActionKeyMapping;
 struct FInputAxisKeyMapping;
 struct FKey;
@@ -29,9 +34,43 @@ public:
 		FName			Name;
 		XrAction		Handle;
 
+		TObjectPtr<const UInputAction> Object;
 		TMap<TPair<XrPath, XrPath>, FName> KeyMap;
+		TMultiMap<TPair<XrPath, XrPath>, TObjectPtr<UInputTrigger>> Triggers;
+		TMultiMap<TPair<XrPath, XrPath>, TObjectPtr<UInputModifier>> Modifiers;
+
+		FOpenXRAction(XrActionSet InActionSet,
+			XrActionType InActionType,
+			const FName& InName, const
+			FString& InLocalizedName,
+			const TArray<XrPath>& InSubactionPaths,
+			const TObjectPtr<const UInputAction>& InObject);
+
+		FOpenXRAction(XrActionSet InActionSet,
+			XrActionType InActionType,
+			const FName& InName,
+			const FString& InLocalizedName,
+			const TArray<XrPath>& InSubactionPaths);
+	};
+
+	struct FOpenXRActionSet
+	{
+		XrActionSet		Handle;
+		FName			Name;
+		FString			LocalizedName;
+
+		TObjectPtr<const UInputMappingContext> Object;
 
-		FOpenXRAction(XrActionSet InActionSet, XrActionType InActionType, const FName& InName, const TArray<XrPath>& SubactionPaths);
+		FOpenXRActionSet(XrInstance InInstance,
+			const FName& InName,
+			const FString& InLocalizedName,
+			uint32 InPriority,
+			const TObjectPtr<const UInputMappingContext>& InObject);
+
+		FOpenXRActionSet(XrInstance InInstance,
+			const FName& InName,
+			const FString& InLocalizedName,
+			uint32 InPriority);
 	};
 
 	struct FOpenXRController
@@ -95,12 +134,19 @@ public:
 
 		FOpenXRHMD* OpenXRHMD;
 
-		TArray<XrActiveActionSet> ActionSets;
-		TArray<XrActionSet> PluginActionSets;
+		TMap<FString, FInteractionProfile> Profiles;
+		TArray<FOpenXRActionSet> ActionSets;
+		TArray<FOpenXRActionSet> PluginActionSets;
 		TArray<XrPath> SubactionPaths;
-		TArray<FOpenXRAction> Actions;
+		TArray<FOpenXRAction> LegacyActions, EnhancedActions;
 		TMap<EControllerHand, FOpenXRController> Controllers;
 		TMap<FName, EControllerHand> MotionSourceToControllerHandMap;
+
+		XrActiveActionSet ControllerSet;
+		XrActiveActionSet LegacySet;
+
+		UPlayerMappableInputConfig* MappableInputConfig;
+
 		XrAction GetActionForMotionSource(FName MotionSource) const;
 		int32 GetDeviceIDForMotionSource(FName MotionSource) const;
 		XrPath GetUserPathForMotionSource(FName MotionSource) const;
@@ -110,12 +156,14 @@ public:
 		bool bDirectionalBindingSupported;
 
 		void BuildActions();
+		void BuildLegacyActions();
+		void BuildEnhancedActions(const FSoftObjectPath& MappableInputConfigPath);
 		void DestroyActions();
 
 		template<typename T>
-		int32 SuggestBindings(XrInstance Instance, FOpenXRAction& Action, const TArray<T>& Mappings, TMap<FString, FInteractionProfile>& Profiles);
-
-		bool SuggestBindingForKey(XrInstance Instance, FOpenXRAction& Action, const FKey& Key, TMap<FString, FInteractionProfile>& Profiles);
+		int32 SuggestBindings(XrInstance Instance, FOpenXRAction& Action, const TArray<T>& Mappings);
+		bool SuggestBindingForKey(XrInstance Instance, FOpenXRAction& Action, const FKey& Key, const TArray<UInputModifier*>& Modifiers, const TArray<UInputTrigger*>& Triggers);
+		bool SuggestBindingForKey(XrInstance Instance, FOpenXRAction& Action, const FKey& Key);
 
 		/** handler to send all messages to */
 		TSharedRef<FGenericApplicationMessageHandler> MessageHandler;
diff --git a/Engine/Plugins/Runtime/OpenXR/Source/OpenXRInput/Private/OpenXRInputSettings.cpp b/Engine/Plugins/Runtime/OpenXR/Source/OpenXRInput/Private/OpenXRInputSettings.cpp
new file mode 100644
--- /dev/null
+++ b/Engine/Plugins/Runtime/OpenXR/Source/OpenXRInput/Private/OpenXRInputSettings.cpp
@@ -0,0 +1,8 @@
+// Copyright Epic Games, Inc. All Rights Reserved.
+
+#include "OpenXRInputSettings.h"
+
+UOpenXRInputSettings::UOpenXRInputSettings(const FObjectInitializer& ObjectInitializer)
+	: Super(ObjectInitializer)
+{
+}
diff --git a/Engine/Plugins/Runtime/OpenXR/Source/OpenXRInput/Public/OpenXRInputSettings.h b/Engine/Plugins/Runtime/OpenXR/Source/OpenXRInput/Public/OpenXRInputSettings.h
new file mode 100644
--- /dev/null
+++ b/Engine/Plugins/Runtime/OpenXR/Source/OpenXRInput/Public/OpenXRInputSettings.h
@@ -0,0 +1,25 @@
+// Copyright Epic Games, Inc. All Rights Reserved.
+
+#pragma once
+
+#include "CoreMinimal.h"
+#include "UObject/Object.h"
+#include "OpenXRInputSettings.generated.h"
+
+/**
+* Implements the settings for the OpenXR Input plugin.
+*/
+UCLASS(config = Input, defaultconfig)
+class OPENXRINPUT_API UOpenXRInputSettings : public UObject
+{
+	GENERATED_UCLASS_BODY()
+
+public:
+	/** Use Enhanced Input for XR instead of regular input. */
+	UPROPERTY(config, EditAnywhere, Category = "Enhanced Input")
+	bool bUseEnhancedInput = false;
+
+	/** A the mappable input config used to generate action sets for OpenXR. */
+	UPROPERTY(config, EditAnywhere, Category = "Enhanced Input", meta = (DisplayName = "Mappable Input Config for XR", AllowedClasses = "PlayerMappableInputConfig"))
+	FSoftObjectPath MappableInputConfig = nullptr;
+};
diff --git a/Engine/Source/Runtime/InputCore/Classes/InputCoreTypes.h b/Engine/Source/Runtime/InputCore/Classes/InputCoreTypes.h
--- a/Engine/Source/Runtime/InputCore/Classes/InputCoreTypes.h
+++ b/Engine/Source/Runtime/InputCore/Classes/InputCoreTypes.h
@@ -508,6 +508,7 @@ struct INPUTCORE_API EKeys
 	static const FKey Vive_Left_Menu_Click;
 	static const FKey Vive_Left_Trigger_Click;
 	static const FKey Vive_Left_Trigger_Axis;
+	static const FKey Vive_Left_Trackpad_2D;
 	static const FKey Vive_Left_Trackpad_X;
 	static const FKey Vive_Left_Trackpad_Y;
 	static const FKey Vive_Left_Trackpad_Click;
@@ -521,6 +522,7 @@ struct INPUTCORE_API EKeys
 	static const FKey Vive_Right_Menu_Click;
 	static const FKey Vive_Right_Trigger_Click;
 	static const FKey Vive_Right_Trigger_Axis;
+	static const FKey Vive_Right_Trackpad_2D;
 	static const FKey Vive_Right_Trackpad_X;
 	static const FKey Vive_Right_Trackpad_Y;
 	static const FKey Vive_Right_Trackpad_Click;
@@ -535,6 +537,7 @@ struct INPUTCORE_API EKeys
 	static const FKey MixedReality_Left_Grip_Click;
 	static const FKey MixedReality_Left_Trigger_Click;
 	static const FKey MixedReality_Left_Trigger_Axis;
+	static const FKey MixedReality_Left_Thumbstick_2D;
 	static const FKey MixedReality_Left_Thumbstick_X;
 	static const FKey MixedReality_Left_Thumbstick_Y;
 	static const FKey MixedReality_Left_Thumbstick_Click;
@@ -542,6 +545,7 @@ struct INPUTCORE_API EKeys
 	static const FKey MixedReality_Left_Thumbstick_Down;
 	static const FKey MixedReality_Left_Thumbstick_Left;
 	static const FKey MixedReality_Left_Thumbstick_Right;
+	static const FKey MixedReality_Left_Trackpad_2D;
 	static const FKey MixedReality_Left_Trackpad_X;
 	static const FKey MixedReality_Left_Trackpad_Y;
 	static const FKey MixedReality_Left_Trackpad_Click;
@@ -554,6 +558,7 @@ struct INPUTCORE_API EKeys
 	static const FKey MixedReality_Right_Grip_Click;
 	static const FKey MixedReality_Right_Trigger_Click;
 	static const FKey MixedReality_Right_Trigger_Axis;
+	static const FKey MixedReality_Right_Thumbstick_2D;
 	static const FKey MixedReality_Right_Thumbstick_X;
 	static const FKey MixedReality_Right_Thumbstick_Y;
 	static const FKey MixedReality_Right_Thumbstick_Click;
@@ -561,6 +566,7 @@ struct INPUTCORE_API EKeys
 	static const FKey MixedReality_Right_Thumbstick_Down;
 	static const FKey MixedReality_Right_Thumbstick_Left;
 	static const FKey MixedReality_Right_Thumbstick_Right;
+	static const FKey MixedReality_Right_Trackpad_2D;
 	static const FKey MixedReality_Right_Trackpad_X;
 	static const FKey MixedReality_Right_Trackpad_Y;
 	static const FKey MixedReality_Right_Trackpad_Click;
@@ -581,6 +587,7 @@ struct INPUTCORE_API EKeys
 	static const FKey OculusTouch_Left_Trigger_Click;
 	static const FKey OculusTouch_Left_Trigger_Axis;
 	static const FKey OculusTouch_Left_Trigger_Touch;
+	static const FKey OculusTouch_Left_Thumbstick_2D;
 	static const FKey OculusTouch_Left_Thumbstick_X;
 	static const FKey OculusTouch_Left_Thumbstick_Y;
 	static const FKey OculusTouch_Left_Thumbstick_Click;
@@ -599,6 +606,7 @@ struct INPUTCORE_API EKeys
 	static const FKey OculusTouch_Right_Trigger_Click;
 	static const FKey OculusTouch_Right_Trigger_Axis;
 	static const FKey OculusTouch_Right_Trigger_Touch;
+	static const FKey OculusTouch_Right_Thumbstick_2D;
 	static const FKey OculusTouch_Right_Thumbstick_X;
 	static const FKey OculusTouch_Right_Thumbstick_Y;
 	static const FKey OculusTouch_Right_Thumbstick_Click;
@@ -620,6 +628,7 @@ struct INPUTCORE_API EKeys
 	static const FKey ValveIndex_Left_Trigger_Click;
 	static const FKey ValveIndex_Left_Trigger_Axis;
 	static const FKey ValveIndex_Left_Trigger_Touch;
+	static const FKey ValveIndex_Left_Thumbstick_2D;
 	static const FKey ValveIndex_Left_Thumbstick_X;
 	static const FKey ValveIndex_Left_Thumbstick_Y;
 	static const FKey ValveIndex_Left_Thumbstick_Click;
@@ -628,6 +637,7 @@ struct INPUTCORE_API EKeys
 	static const FKey ValveIndex_Left_Thumbstick_Down;
 	static const FKey ValveIndex_Left_Thumbstick_Left;
 	static const FKey ValveIndex_Left_Thumbstick_Right;
+	static const FKey ValveIndex_Left_Trackpad_2D;
 	static const FKey ValveIndex_Left_Trackpad_X;
 	static const FKey ValveIndex_Left_Trackpad_Y;
 	static const FKey ValveIndex_Left_Trackpad_Force;
@@ -647,6 +657,7 @@ struct INPUTCORE_API EKeys
 	static const FKey ValveIndex_Right_Trigger_Click;
 	static const FKey ValveIndex_Right_Trigger_Axis;
 	static const FKey ValveIndex_Right_Trigger_Touch;
+	static const FKey ValveIndex_Right_Thumbstick_2D;
 	static const FKey ValveIndex_Right_Thumbstick_X;
 	static const FKey ValveIndex_Right_Thumbstick_Y;
 	static const FKey ValveIndex_Right_Thumbstick_Click;
@@ -655,6 +666,7 @@ struct INPUTCORE_API EKeys
 	static const FKey ValveIndex_Right_Thumbstick_Down;
 	static const FKey ValveIndex_Right_Thumbstick_Left;
 	static const FKey ValveIndex_Right_Thumbstick_Right;
+	static const FKey ValveIndex_Right_Trackpad_2D;
 	static const FKey ValveIndex_Right_Trackpad_X;
 	static const FKey ValveIndex_Right_Trackpad_Y;
 	static const FKey ValveIndex_Right_Trackpad_Force;
diff --git a/Engine/Source/Runtime/InputCore/Private/InputCoreTypes.cpp b/Engine/Source/Runtime/InputCore/Private/InputCoreTypes.cpp
--- a/Engine/Source/Runtime/InputCore/Private/InputCoreTypes.cpp
+++ b/Engine/Source/Runtime/InputCore/Private/InputCoreTypes.cpp
@@ -263,6 +263,7 @@ const FKey EKeys::Vive_Left_Grip_Click("Vive_Left_Grip_Click");
 const FKey EKeys::Vive_Left_Menu_Click("Vive_Left_Menu_Click");
 const FKey EKeys::Vive_Left_Trigger_Click("Vive_Left_Trigger_Click");
 const FKey EKeys::Vive_Left_Trigger_Axis("Vive_Left_Trigger_Axis");
+const FKey EKeys::Vive_Left_Trackpad_2D("Vive_Left_Trackpad_2D");
 const FKey EKeys::Vive_Left_Trackpad_X("Vive_Left_Trackpad_X");
 const FKey EKeys::Vive_Left_Trackpad_Y("Vive_Left_Trackpad_Y");
 const FKey EKeys::Vive_Left_Trackpad_Click("Vive_Left_Trackpad_Click");
@@ -276,6 +277,7 @@ const FKey EKeys::Vive_Right_Grip_Click("Vive_Right_Grip_Click");
 const FKey EKeys::Vive_Right_Menu_Click("Vive_Right_Menu_Click");
 const FKey EKeys::Vive_Right_Trigger_Click("Vive_Right_Trigger_Click");
 const FKey EKeys::Vive_Right_Trigger_Axis("Vive_Right_Trigger_Axis");
+const FKey EKeys::Vive_Right_Trackpad_2D("Vive_Right_Trackpad_2D");
 const FKey EKeys::Vive_Right_Trackpad_X("Vive_Right_Trackpad_X");
 const FKey EKeys::Vive_Right_Trackpad_Y("Vive_Right_Trackpad_Y");
 const FKey EKeys::Vive_Right_Trackpad_Click("Vive_Right_Trackpad_Click");
@@ -290,6 +292,7 @@ const FKey EKeys::MixedReality_Left_Menu_Click("MixedReality_Left_Menu_Click");
 const FKey EKeys::MixedReality_Left_Grip_Click("MixedReality_Left_Grip_Click");
 const FKey EKeys::MixedReality_Left_Trigger_Click("MixedReality_Left_Trigger_Click");
 const FKey EKeys::MixedReality_Left_Trigger_Axis("MixedReality_Left_Trigger_Axis");
+const FKey EKeys::MixedReality_Left_Thumbstick_2D("MixedReality_Left_Thumbstick_2D");
 const FKey EKeys::MixedReality_Left_Thumbstick_X("MixedReality_Left_Thumbstick_X");
 const FKey EKeys::MixedReality_Left_Thumbstick_Y("MixedReality_Left_Thumbstick_Y");
 const FKey EKeys::MixedReality_Left_Thumbstick_Click("MixedReality_Left_Thumbstick_Click");
@@ -297,6 +300,7 @@ const FKey EKeys::MixedReality_Left_Thumbstick_Up("MixedReality_Left_Thumbstick_
 const FKey EKeys::MixedReality_Left_Thumbstick_Down("MixedReality_Left_Thumbstick_Down");
 const FKey EKeys::MixedReality_Left_Thumbstick_Left("MixedReality_Left_Thumbstick_Left");
 const FKey EKeys::MixedReality_Left_Thumbstick_Right("MixedReality_Left_Thumbstick_Right");
+const FKey EKeys::MixedReality_Left_Trackpad_2D("MixedReality_Left_Trackpad_2D");
 const FKey EKeys::MixedReality_Left_Trackpad_X("MixedReality_Left_Trackpad_X");
 const FKey EKeys::MixedReality_Left_Trackpad_Y("MixedReality_Left_Trackpad_Y");
 const FKey EKeys::MixedReality_Left_Trackpad_Click("MixedReality_Left_Trackpad_Click");
@@ -309,6 +313,7 @@ const FKey EKeys::MixedReality_Right_Menu_Click("MixedReality_Right_Menu_Click")
 const FKey EKeys::MixedReality_Right_Grip_Click("MixedReality_Right_Grip_Click");
 const FKey EKeys::MixedReality_Right_Trigger_Click("MixedReality_Right_Trigger_Click");
 const FKey EKeys::MixedReality_Right_Trigger_Axis("MixedReality_Right_Trigger_Axis");
+const FKey EKeys::MixedReality_Right_Thumbstick_2D("MixedReality_Right_Thumbstick_2D");
 const FKey EKeys::MixedReality_Right_Thumbstick_X("MixedReality_Right_Thumbstick_X");
 const FKey EKeys::MixedReality_Right_Thumbstick_Y("MixedReality_Right_Thumbstick_Y");
 const FKey EKeys::MixedReality_Right_Thumbstick_Click("MixedReality_Right_Thumbstick_Click");
@@ -316,6 +321,7 @@ const FKey EKeys::MixedReality_Right_Thumbstick_Up("MixedReality_Right_Thumbstic
 const FKey EKeys::MixedReality_Right_Thumbstick_Down("MixedReality_Right_Thumbstick_Down");
 const FKey EKeys::MixedReality_Right_Thumbstick_Left("MixedReality_Right_Thumbstick_Left");
 const FKey EKeys::MixedReality_Right_Thumbstick_Right("MixedReality_Right_Thumbstick_Right");
+const FKey EKeys::MixedReality_Right_Trackpad_2D("MixedReality_Right_Trackpad_2D");
 const FKey EKeys::MixedReality_Right_Trackpad_X("MixedReality_Right_Trackpad_X");
 const FKey EKeys::MixedReality_Right_Trackpad_Y("MixedReality_Right_Trackpad_Y");
 const FKey EKeys::MixedReality_Right_Trackpad_Click("MixedReality_Right_Trackpad_Click");
@@ -336,6 +342,7 @@ const FKey EKeys::OculusTouch_Left_Grip_Axis("OculusTouch_Left_Grip_Axis");
 const FKey EKeys::OculusTouch_Left_Trigger_Click("OculusTouch_Left_Trigger_Click");
 const FKey EKeys::OculusTouch_Left_Trigger_Axis("OculusTouch_Left_Trigger_Axis");
 const FKey EKeys::OculusTouch_Left_Trigger_Touch("OculusTouch_Left_Trigger_Touch");
+const FKey EKeys::OculusTouch_Left_Thumbstick_2D("OculusTouch_Left_Thumbstick_2D");
 const FKey EKeys::OculusTouch_Left_Thumbstick_X("OculusTouch_Left_Thumbstick_X");
 const FKey EKeys::OculusTouch_Left_Thumbstick_Y("OculusTouch_Left_Thumbstick_Y");
 const FKey EKeys::OculusTouch_Left_Thumbstick_Click("OculusTouch_Left_Thumbstick_Click");
@@ -354,6 +361,7 @@ const FKey EKeys::OculusTouch_Right_Grip_Axis("OculusTouch_Right_Grip_Axis");
 const FKey EKeys::OculusTouch_Right_Trigger_Click("OculusTouch_Right_Trigger_Click");
 const FKey EKeys::OculusTouch_Right_Trigger_Axis("OculusTouch_Right_Trigger_Axis");
 const FKey EKeys::OculusTouch_Right_Trigger_Touch("OculusTouch_Right_Trigger_Touch");
+const FKey EKeys::OculusTouch_Right_Thumbstick_2D("OculusTouch_Right_Thumbstick_2D");
 const FKey EKeys::OculusTouch_Right_Thumbstick_X("OculusTouch_Right_Thumbstick_X");
 const FKey EKeys::OculusTouch_Right_Thumbstick_Y("OculusTouch_Right_Thumbstick_Y");
 const FKey EKeys::OculusTouch_Right_Thumbstick_Click("OculusTouch_Right_Thumbstick_Click");
@@ -375,6 +383,7 @@ const FKey EKeys::ValveIndex_Left_Grip_Force("ValveIndex_Left_Grip_Force");
 const FKey EKeys::ValveIndex_Left_Trigger_Click("ValveIndex_Left_Trigger_Click");
 const FKey EKeys::ValveIndex_Left_Trigger_Axis("ValveIndex_Left_Trigger_Axis");
 const FKey EKeys::ValveIndex_Left_Trigger_Touch("ValveIndex_Left_Trigger_Touch");
+const FKey EKeys::ValveIndex_Left_Thumbstick_2D("ValveIndex_Left_Thumbstick_2D");
 const FKey EKeys::ValveIndex_Left_Thumbstick_X("ValveIndex_Left_Thumbstick_X");
 const FKey EKeys::ValveIndex_Left_Thumbstick_Y("ValveIndex_Left_Thumbstick_Y");
 const FKey EKeys::ValveIndex_Left_Thumbstick_Click("ValveIndex_Left_Thumbstick_Click");
@@ -383,6 +392,7 @@ const FKey EKeys::ValveIndex_Left_Thumbstick_Up("ValveIndex_Left_Thumbstick_Up")
 const FKey EKeys::ValveIndex_Left_Thumbstick_Down("ValveIndex_Left_Thumbstick_Down");
 const FKey EKeys::ValveIndex_Left_Thumbstick_Left("ValveIndex_Left_Thumbstick_Left");
 const FKey EKeys::ValveIndex_Left_Thumbstick_Right("ValveIndex_Left_Thumbstick_Right");
+const FKey EKeys::ValveIndex_Left_Trackpad_2D("ValveIndex_Left_Trackpad_2D");
 const FKey EKeys::ValveIndex_Left_Trackpad_X("ValveIndex_Left_Trackpad_X");
 const FKey EKeys::ValveIndex_Left_Trackpad_Y("ValveIndex_Left_Trackpad_Y");
 const FKey EKeys::ValveIndex_Left_Trackpad_Force("ValveIndex_Left_Trackpad_Force");
@@ -402,6 +412,7 @@ const FKey EKeys::ValveIndex_Right_Grip_Force("ValveIndex_Right_Grip_Force");
 const FKey EKeys::ValveIndex_Right_Trigger_Click("ValveIndex_Right_Trigger_Click");
 const FKey EKeys::ValveIndex_Right_Trigger_Axis("ValveIndex_Right_Trigger_Axis");
 const FKey EKeys::ValveIndex_Right_Trigger_Touch("ValveIndex_Right_Trigger_Touch");
+const FKey EKeys::ValveIndex_Right_Thumbstick_2D("ValveIndex_Right_Thumbstick_2D");
 const FKey EKeys::ValveIndex_Right_Thumbstick_X("ValveIndex_Right_Thumbstick_X");
 const FKey EKeys::ValveIndex_Right_Thumbstick_Y("ValveIndex_Right_Thumbstick_Y");
 const FKey EKeys::ValveIndex_Right_Thumbstick_Click("ValveIndex_Right_Thumbstick_Click");
@@ -410,6 +421,7 @@ const FKey EKeys::ValveIndex_Right_Thumbstick_Up("ValveIndex_Right_Thumbstick_Up
 const FKey EKeys::ValveIndex_Right_Thumbstick_Down("ValveIndex_Right_Thumbstick_Down");
 const FKey EKeys::ValveIndex_Right_Thumbstick_Left("ValveIndex_Right_Thumbstick_Left");
 const FKey EKeys::ValveIndex_Right_Thumbstick_Right("ValveIndex_Right_Thumbstick_Right");
+const FKey EKeys::ValveIndex_Right_Trackpad_2D("ValveIndex_Right_Trackpad_2D");
 const FKey EKeys::ValveIndex_Right_Trackpad_X("ValveIndex_Right_Trackpad_X");
 const FKey EKeys::ValveIndex_Right_Trackpad_Y("ValveIndex_Right_Trackpad_Y");
 const FKey EKeys::ValveIndex_Right_Trackpad_Force("ValveIndex_Right_Trackpad_Force");
@@ -780,9 +792,10 @@ void EKeys::Initialize()
 	AddKey(FKeyDetails(EKeys::Vive_Left_Menu_Click, LOCTEXT("Vive_Left_Menu_Click", "Vive (L) Menu"), FKeyDetails::GamepadKey | FKeyDetails::NotBlueprintBindableKey, "Vive"));
 	AddKey(FKeyDetails(EKeys::Vive_Left_Trigger_Click, LOCTEXT("Vive_Left_Trigger_Click", "Vive (L) Trigger"), FKeyDetails::GamepadKey | FKeyDetails::NotBlueprintBindableKey, "Vive"));
 	AddKey(FKeyDetails(EKeys::Vive_Left_Trigger_Axis, LOCTEXT("Vive_Left_Trigger_Axis", "Vive (L) Trigger Axis"), FKeyDetails::GamepadKey | FKeyDetails::Axis1D | FKeyDetails::NotBlueprintBindableKey, "Vive"));
-	AddKey(FKeyDetails(EKeys::Vive_Left_Trackpad_X, LOCTEXT("Vive_Left_Trackpad_X", "Vive (L) Trackpad X"), FKeyDetails::GamepadKey | FKeyDetails::Axis1D | FKeyDetails::NotBlueprintBindableKey, "Vive"));
-	AddKey(FKeyDetails(EKeys::Vive_Left_Trackpad_Y, LOCTEXT("Vive_Left_Trackpad_Y", "Vive (L) Trackpad Y"), FKeyDetails::GamepadKey | FKeyDetails::Axis1D | FKeyDetails::NotBlueprintBindableKey, "Vive"));
-	AddKey(FKeyDetails(EKeys::Vive_Left_Trackpad_Click, LOCTEXT("Vive_Left_Trackpad_Click", "Vive (L) Trackpad"), FKeyDetails::GamepadKey | FKeyDetails::NotBlueprintBindableKey, "Vive"));
+	AddKey(FKeyDetails(EKeys::Vive_Left_Trackpad_2D, LOCTEXT("Vive_Left_Trackpad_2D", "Vive (L) Trackpad 2D-Axis"), FKeyDetails::GamepadKey | FKeyDetails::Axis2D | FKeyDetails::NotBlueprintBindableKey, "Vive"));
+	AddKey(FKeyDetails(EKeys::Vive_Left_Trackpad_X, LOCTEXT("Vive_Left_Trackpad_X", "Vive (L) Trackpad X-Axis"), FKeyDetails::GamepadKey | FKeyDetails::Axis1D | FKeyDetails::NotBlueprintBindableKey, "Vive"));
+	AddKey(FKeyDetails(EKeys::Vive_Left_Trackpad_Y, LOCTEXT("Vive_Left_Trackpad_Y", "Vive (L) Trackpad Y-Axis"), FKeyDetails::GamepadKey | FKeyDetails::Axis1D | FKeyDetails::NotBlueprintBindableKey, "Vive"));
+	AddKey(FKeyDetails(EKeys::Vive_Left_Trackpad_Click, LOCTEXT("Vive_Left_Trackpad_Click", "Vive (L) Trackpad Button"), FKeyDetails::GamepadKey | FKeyDetails::NotBlueprintBindableKey, "Vive"));
 	AddKey(FKeyDetails(EKeys::Vive_Left_Trackpad_Touch, LOCTEXT("Vive_Left_Trackpad_Touch", "Vive (L) Trackpad Touch"), FKeyDetails::GamepadKey | FKeyDetails::NotBlueprintBindableKey, "Vive"));
 	AddKey(FKeyDetails(EKeys::Vive_Left_Trackpad_Up, LOCTEXT("Vive_Left_Trackpad_Up", "Vive (L) Trackpad Up"), FKeyDetails::GamepadKey | FKeyDetails::NotBlueprintBindableKey, "Vive"));
 	AddKey(FKeyDetails(EKeys::Vive_Left_Trackpad_Down, LOCTEXT("Vive_Left_Trackpad_Down", "Vive (L) Trackpad Down"), FKeyDetails::GamepadKey | FKeyDetails::NotBlueprintBindableKey, "Vive"));
@@ -793,9 +806,10 @@ void EKeys::Initialize()
 	AddKey(FKeyDetails(EKeys::Vive_Right_Menu_Click, LOCTEXT("Vive_Right_Menu_Click", "Vive (R) Menu"), FKeyDetails::GamepadKey | FKeyDetails::NotBlueprintBindableKey, "Vive"));
 	AddKey(FKeyDetails(EKeys::Vive_Right_Trigger_Click, LOCTEXT("Vive_Right_Trigger_Click", "Vive (R) Trigger"), FKeyDetails::GamepadKey | FKeyDetails::NotBlueprintBindableKey, "Vive"));
 	AddKey(FKeyDetails(EKeys::Vive_Right_Trigger_Axis, LOCTEXT("Vive_Right_Trigger_Axis", "Vive (R) Trigger Axis"), FKeyDetails::GamepadKey | FKeyDetails::Axis1D | FKeyDetails::NotBlueprintBindableKey, "Vive"));
-	AddKey(FKeyDetails(EKeys::Vive_Right_Trackpad_X, LOCTEXT("Vive_Right_Trackpad_X", "Vive (R) Trackpad X"), FKeyDetails::GamepadKey | FKeyDetails::Axis1D | FKeyDetails::NotBlueprintBindableKey, "Vive"));
-	AddKey(FKeyDetails(EKeys::Vive_Right_Trackpad_Y, LOCTEXT("Vive_Right_Trackpad_Y", "Vive (R) Trackpad Y"), FKeyDetails::GamepadKey | FKeyDetails::Axis1D | FKeyDetails::NotBlueprintBindableKey, "Vive"));
-	AddKey(FKeyDetails(EKeys::Vive_Right_Trackpad_Click, LOCTEXT("Vive_Right_Trackpad_Click", "Vive (R) Trackpad"), FKeyDetails::GamepadKey | FKeyDetails::NotBlueprintBindableKey, "Vive"));
+	AddKey(FKeyDetails(EKeys::Vive_Right_Trackpad_2D, LOCTEXT("Vive_Right_Trackpad_2D", "Vive (R) Trackpad 2D-Axis"), FKeyDetails::GamepadKey | FKeyDetails::Axis2D | FKeyDetails::NotBlueprintBindableKey, "Vive"));
+	AddKey(FKeyDetails(EKeys::Vive_Right_Trackpad_X, LOCTEXT("Vive_Right_Trackpad_X", "Vive (R) Trackpad X-Axis"), FKeyDetails::GamepadKey | FKeyDetails::Axis1D | FKeyDetails::NotBlueprintBindableKey, "Vive"));
+	AddKey(FKeyDetails(EKeys::Vive_Right_Trackpad_Y, LOCTEXT("Vive_Right_Trackpad_Y", "Vive (R) Trackpad Y-Axis"), FKeyDetails::GamepadKey | FKeyDetails::Axis1D | FKeyDetails::NotBlueprintBindableKey, "Vive"));
+	AddKey(FKeyDetails(EKeys::Vive_Right_Trackpad_Click, LOCTEXT("Vive_Right_Trackpad_Click", "Vive (R) Trackpad Button"), FKeyDetails::GamepadKey | FKeyDetails::NotBlueprintBindableKey, "Vive"));
 	AddKey(FKeyDetails(EKeys::Vive_Right_Trackpad_Touch, LOCTEXT("Vive_Right_Trackpad_Touch", "Vive (R) Trackpad Touch"), FKeyDetails::GamepadKey | FKeyDetails::NotBlueprintBindableKey, "Vive"));
 	AddKey(FKeyDetails(EKeys::Vive_Right_Trackpad_Up, LOCTEXT("Vive_Right_Trackpad_Up", "Vive (R) Trackpad Up"), FKeyDetails::GamepadKey | FKeyDetails::NotBlueprintBindableKey, "Vive"));
 	AddKey(FKeyDetails(EKeys::Vive_Right_Trackpad_Down, LOCTEXT("Vive_Right_Trackpad_Down", "Vive (R) Trackpad Down"), FKeyDetails::GamepadKey | FKeyDetails::NotBlueprintBindableKey, "Vive"));
@@ -809,16 +823,18 @@ void EKeys::Initialize()
 	AddKey(FKeyDetails(EKeys::MixedReality_Left_Grip_Click, LOCTEXT("MixedReality_Left_Grip_Click", "Mixed Reality (L) Grip"), FKeyDetails::GamepadKey | FKeyDetails::NotBlueprintBindableKey, "MixedReality"));
 	AddKey(FKeyDetails(EKeys::MixedReality_Left_Trigger_Click, LOCTEXT("MixedReality_Left_Trigger_Click", "Mixed Reality (L) Trigger"), FKeyDetails::GamepadKey | FKeyDetails::NotBlueprintBindableKey, "MixedReality"));
 	AddKey(FKeyDetails(EKeys::MixedReality_Left_Trigger_Axis, LOCTEXT("MixedReality_Left_Trigger_Axis", "Mixed Reality (L) Trigger Axis"), FKeyDetails::GamepadKey | FKeyDetails::Axis1D | FKeyDetails::NotBlueprintBindableKey, "MixedReality"));
-	AddKey(FKeyDetails(EKeys::MixedReality_Left_Thumbstick_X, LOCTEXT("MixedReality_Left_Thumbstick_X", "Mixed Reality (L) Thumbstick X"), FKeyDetails::GamepadKey | FKeyDetails::Axis1D | FKeyDetails::NotBlueprintBindableKey, "MixedReality"));
-	AddKey(FKeyDetails(EKeys::MixedReality_Left_Thumbstick_Y, LOCTEXT("MixedReality_Left_Thumbstick_Y", "Mixed Reality (L) Thumbstick Y"), FKeyDetails::GamepadKey | FKeyDetails::Axis1D | FKeyDetails::NotBlueprintBindableKey, "MixedReality"));
-	AddKey(FKeyDetails(EKeys::MixedReality_Left_Thumbstick_Click, LOCTEXT("MixedReality_Left_Thumbstick_Click", "Mixed Reality (L) Thumbstick"), FKeyDetails::GamepadKey | FKeyDetails::NotBlueprintBindableKey, "MixedReality"));
+	AddKey(FKeyDetails(EKeys::MixedReality_Left_Thumbstick_2D, LOCTEXT("MixedReality_Left_Thumbstick_2D", "Mixed Reality (L) Thumbstick 2D-Axis"), FKeyDetails::GamepadKey | FKeyDetails::Axis2D | FKeyDetails::NotBlueprintBindableKey, "MixedReality"));
+	AddKey(FKeyDetails(EKeys::MixedReality_Left_Thumbstick_X, LOCTEXT("MixedReality_Left_Thumbstick_X", "Mixed Reality (L) Thumbstick X-Axis"), FKeyDetails::GamepadKey | FKeyDetails::Axis1D | FKeyDetails::NotBlueprintBindableKey, "MixedReality"));
+	AddKey(FKeyDetails(EKeys::MixedReality_Left_Thumbstick_Y, LOCTEXT("MixedReality_Left_Thumbstick_Y", "Mixed Reality (L) Thumbstick Y-Axis"), FKeyDetails::GamepadKey | FKeyDetails::Axis1D | FKeyDetails::NotBlueprintBindableKey, "MixedReality"));
+	AddKey(FKeyDetails(EKeys::MixedReality_Left_Thumbstick_Click, LOCTEXT("MixedReality_Left_Thumbstick_Click", "Mixed Reality (L) Thumbstick Button"), FKeyDetails::GamepadKey | FKeyDetails::NotBlueprintBindableKey, "MixedReality"));
 	AddKey(FKeyDetails(EKeys::MixedReality_Left_Thumbstick_Up, LOCTEXT("MixedReality_Left_Thumbstick_Up", "Mixed Reality (L) Thumbstick Up"), FKeyDetails::GamepadKey | FKeyDetails::ButtonAxis | FKeyDetails::NotBlueprintBindableKey, "MixedReality"));
 	AddKey(FKeyDetails(EKeys::MixedReality_Left_Thumbstick_Down, LOCTEXT("MixedReality_Left_Thumbstick_Down", "Mixed Reality (L) Thumbstick Down"), FKeyDetails::GamepadKey | FKeyDetails::ButtonAxis | FKeyDetails::NotBlueprintBindableKey, "MixedReality"));
 	AddKey(FKeyDetails(EKeys::MixedReality_Left_Thumbstick_Left, LOCTEXT("MixedReality_Left_Thumbstick_Left", "Mixed Reality (L) Thumbstick Left"), FKeyDetails::GamepadKey | FKeyDetails::ButtonAxis | FKeyDetails::NotBlueprintBindableKey, "MixedReality"));
 	AddKey(FKeyDetails(EKeys::MixedReality_Left_Thumbstick_Right, LOCTEXT("MixedReality_Left_Thumbstick_Right", "Mixed Reality (L) Thumbstick Right"), FKeyDetails::GamepadKey | FKeyDetails::ButtonAxis | FKeyDetails::NotBlueprintBindableKey, "MixedReality"));
+	AddKey(FKeyDetails(EKeys::MixedReality_Left_Trackpad_2D, LOCTEXT("MixedReality_Left_Trackpad_2D", "Mixed Reality (L) Trackpad 2D-Axis"), FKeyDetails::GamepadKey | FKeyDetails::Axis2D | FKeyDetails::NotBlueprintBindableKey, "MixedReality"));
 	AddKey(FKeyDetails(EKeys::MixedReality_Left_Trackpad_X, LOCTEXT("MixedReality_Left_Trackpad_X", "Mixed Reality (L) Trackpad X"), FKeyDetails::GamepadKey | FKeyDetails::Axis1D | FKeyDetails::NotBlueprintBindableKey, "MixedReality"));
 	AddKey(FKeyDetails(EKeys::MixedReality_Left_Trackpad_Y, LOCTEXT("MixedReality_Left_Trackpad_Y", "Mixed Reality (L) Trackpad Y"), FKeyDetails::GamepadKey | FKeyDetails::Axis1D | FKeyDetails::NotBlueprintBindableKey, "MixedReality"));
-	AddKey(FKeyDetails(EKeys::MixedReality_Left_Trackpad_Click, LOCTEXT("MixedReality_Left_Trackpad_Click", "Mixed Reality (L) Trackpad"), FKeyDetails::GamepadKey | FKeyDetails::NotBlueprintBindableKey, "MixedReality"));
+	AddKey(FKeyDetails(EKeys::MixedReality_Left_Trackpad_Click, LOCTEXT("MixedReality_Left_Trackpad_Click", "Mixed Reality (L) Trackpad Button"), FKeyDetails::GamepadKey | FKeyDetails::NotBlueprintBindableKey, "MixedReality"));
 	AddKey(FKeyDetails(EKeys::MixedReality_Left_Trackpad_Touch, LOCTEXT("MixedReality_Left_Trackpad_Touch", "Mixed Reality (L) Trackpad Touch"), FKeyDetails::GamepadKey | FKeyDetails::NotBlueprintBindableKey, "MixedReality"));
 	AddKey(FKeyDetails(EKeys::MixedReality_Left_Trackpad_Up, LOCTEXT("MixedReality_Left_Trackpad_Up", "Mixed Reality (L) Trackpad Up"), FKeyDetails::GamepadKey | FKeyDetails::NotBlueprintBindableKey, "MixedReality"));
 	AddKey(FKeyDetails(EKeys::MixedReality_Left_Trackpad_Down, LOCTEXT("MixedReality_Left_Trackpad_Down", "Mixed Reality (L) Trackpad Down"), FKeyDetails::GamepadKey | FKeyDetails::NotBlueprintBindableKey, "MixedReality"));
@@ -828,16 +844,18 @@ void EKeys::Initialize()
 	AddKey(FKeyDetails(EKeys::MixedReality_Right_Grip_Click, LOCTEXT("MixedReality_Right_Grip_Click", "Mixed Reality (R) Grip"), FKeyDetails::GamepadKey | FKeyDetails::NotBlueprintBindableKey, "MixedReality"));
 	AddKey(FKeyDetails(EKeys::MixedReality_Right_Trigger_Click, LOCTEXT("MixedReality_Right_Trigger_Click", "Mixed Reality (R) Trigger"), FKeyDetails::GamepadKey | FKeyDetails::NotBlueprintBindableKey, "MixedReality"));
 	AddKey(FKeyDetails(EKeys::MixedReality_Right_Trigger_Axis, LOCTEXT("MixedReality_Right_Trigger_Axis", "Mixed Reality (R) Trigger Axis"), FKeyDetails::GamepadKey | FKeyDetails::Axis1D | FKeyDetails::NotBlueprintBindableKey, "MixedReality"));
-	AddKey(FKeyDetails(EKeys::MixedReality_Right_Thumbstick_X, LOCTEXT("MixedReality_Right_Thumbstick_X", "Mixed Reality (R) Thumbstick X"), FKeyDetails::GamepadKey | FKeyDetails::Axis1D | FKeyDetails::NotBlueprintBindableKey, "MixedReality"));
-	AddKey(FKeyDetails(EKeys::MixedReality_Right_Thumbstick_Y, LOCTEXT("MixedReality_Right_Thumbstick_Y", "Mixed Reality (R) Thumbstick Y"), FKeyDetails::GamepadKey | FKeyDetails::Axis1D | FKeyDetails::NotBlueprintBindableKey, "MixedReality"));
-	AddKey(FKeyDetails(EKeys::MixedReality_Right_Thumbstick_Click, LOCTEXT("MixedReality_Right_Thumbstick_Click", "Mixed Reality (R) Thumbstick"), FKeyDetails::GamepadKey | FKeyDetails::NotBlueprintBindableKey, "MixedReality"));
+	AddKey(FKeyDetails(EKeys::MixedReality_Right_Thumbstick_2D, LOCTEXT("MixedReality_Right_Thumbstick_2D", "Mixed Reality (R) Thumbstick 2D-Axis"), FKeyDetails::GamepadKey | FKeyDetails::Axis2D | FKeyDetails::NotBlueprintBindableKey, "MixedReality"));
+	AddKey(FKeyDetails(EKeys::MixedReality_Right_Thumbstick_X, LOCTEXT("MixedReality_Right_Thumbstick_X", "Mixed Reality (R) Thumbstick X-Axis"), FKeyDetails::GamepadKey | FKeyDetails::Axis1D | FKeyDetails::NotBlueprintBindableKey, "MixedReality"));
+	AddKey(FKeyDetails(EKeys::MixedReality_Right_Thumbstick_Y, LOCTEXT("MixedReality_Right_Thumbstick_Y", "Mixed Reality (R) Thumbstick Y-Axis"), FKeyDetails::GamepadKey | FKeyDetails::Axis1D | FKeyDetails::NotBlueprintBindableKey, "MixedReality"));
+	AddKey(FKeyDetails(EKeys::MixedReality_Right_Thumbstick_Click, LOCTEXT("MixedReality_Right_Thumbstick_Click", "Mixed Reality (R) Thumbstick Button"), FKeyDetails::GamepadKey | FKeyDetails::NotBlueprintBindableKey, "MixedReality"));
 	AddKey(FKeyDetails(EKeys::MixedReality_Right_Thumbstick_Up, LOCTEXT("MixedReality_Right_Thumbstick_Up", "Mixed Reality (R) Thumbstick Up"), FKeyDetails::GamepadKey | FKeyDetails::ButtonAxis | FKeyDetails::NotBlueprintBindableKey, "MixedReality"));
 	AddKey(FKeyDetails(EKeys::MixedReality_Right_Thumbstick_Down, LOCTEXT("MixedReality_Right_Thumbstick_Down", "Mixed Reality (R) Thumbstick Down"), FKeyDetails::GamepadKey | FKeyDetails::ButtonAxis | FKeyDetails::NotBlueprintBindableKey, "MixedReality"));
 	AddKey(FKeyDetails(EKeys::MixedReality_Right_Thumbstick_Left, LOCTEXT("MixedReality_Right_Thumbstick_Left", "Mixed Reality (R) Thumbstick Left"), FKeyDetails::GamepadKey | FKeyDetails::ButtonAxis | FKeyDetails::NotBlueprintBindableKey, "MixedReality"));
 	AddKey(FKeyDetails(EKeys::MixedReality_Right_Thumbstick_Right, LOCTEXT("MixedReality_Right_Thumbstick_Right", "Mixed Reality (R) Thumbstick Right"), FKeyDetails::GamepadKey | FKeyDetails::ButtonAxis | FKeyDetails::NotBlueprintBindableKey, "MixedReality"));
+	AddKey(FKeyDetails(EKeys::MixedReality_Right_Trackpad_2D, LOCTEXT("MixedReality_Right_Trackpad_2D", "Mixed Reality (R) Trackpad 2D-Axis"), FKeyDetails::GamepadKey | FKeyDetails::Axis2D | FKeyDetails::NotBlueprintBindableKey, "MixedReality"));
 	AddKey(FKeyDetails(EKeys::MixedReality_Right_Trackpad_X, LOCTEXT("MixedReality_Right_Trackpad_X", "Mixed Reality (R) Trackpad X"), FKeyDetails::GamepadKey | FKeyDetails::Axis1D | FKeyDetails::NotBlueprintBindableKey, "MixedReality"));
 	AddKey(FKeyDetails(EKeys::MixedReality_Right_Trackpad_Y, LOCTEXT("MixedReality_Right_Trackpad_Y", "Mixed Reality (R) Trackpad Y"), FKeyDetails::GamepadKey | FKeyDetails::Axis1D | FKeyDetails::NotBlueprintBindableKey, "MixedReality"));
-	AddKey(FKeyDetails(EKeys::MixedReality_Right_Trackpad_Click, LOCTEXT("MixedReality_Right_Trackpad_Click", "Mixed Reality (R) Trackpad"), FKeyDetails::GamepadKey | FKeyDetails::NotBlueprintBindableKey, "MixedReality"));
+	AddKey(FKeyDetails(EKeys::MixedReality_Right_Trackpad_Click, LOCTEXT("MixedReality_Right_Trackpad_Click", "Mixed Reality (R) Trackpad Button"), FKeyDetails::GamepadKey | FKeyDetails::NotBlueprintBindableKey, "MixedReality"));
 	AddKey(FKeyDetails(EKeys::MixedReality_Right_Trackpad_Touch, LOCTEXT("MixedReality_Right_Trackpad_Touch", "Mixed Reality (R) Trackpad Touch"), FKeyDetails::GamepadKey | FKeyDetails::NotBlueprintBindableKey, "MixedReality"));
 	AddKey(FKeyDetails(EKeys::MixedReality_Right_Trackpad_Up, LOCTEXT("MixedReality_Right_Trackpad_Up", "Mixed Reality (R) Trackpad Up"), FKeyDetails::GamepadKey | FKeyDetails::NotBlueprintBindableKey, "MixedReality"));
 	AddKey(FKeyDetails(EKeys::MixedReality_Right_Trackpad_Down, LOCTEXT("MixedReality_Right_Trackpad_Down", "Mixed Reality (R) Trackpad Down"), FKeyDetails::GamepadKey | FKeyDetails::NotBlueprintBindableKey, "MixedReality"));
@@ -857,9 +875,10 @@ void EKeys::Initialize()
 	AddKey(FKeyDetails(EKeys::OculusTouch_Left_Trigger_Click, LOCTEXT("OculusTouch_Left_Trigger_Click", "Oculus Touch (L) Trigger"), FKeyDetails::GamepadKey | FKeyDetails::NotBlueprintBindableKey, "OculusTouch"));
 	AddKey(FKeyDetails(EKeys::OculusTouch_Left_Trigger_Axis, LOCTEXT("OculusTouch_Left_Trigger_Axis", "Oculus Touch (L) Trigger Axis"), FKeyDetails::GamepadKey | FKeyDetails::Axis1D | FKeyDetails::NotBlueprintBindableKey, "OculusTouch"));
 	AddKey(FKeyDetails(EKeys::OculusTouch_Left_Trigger_Touch, LOCTEXT("OculusTouch_Left_Trigger_Touch", "Oculus Touch (L) Trigger Touch"), FKeyDetails::GamepadKey | FKeyDetails::NotBlueprintBindableKey, "OculusTouch"));
-	AddKey(FKeyDetails(EKeys::OculusTouch_Left_Thumbstick_X, LOCTEXT("OculusTouch_Left_Thumbstick_X", "Oculus Touch (L) Thumbstick X"), FKeyDetails::GamepadKey | FKeyDetails::Axis1D | FKeyDetails::NotBlueprintBindableKey, "OculusTouch"));
-	AddKey(FKeyDetails(EKeys::OculusTouch_Left_Thumbstick_Y, LOCTEXT("OculusTouch_Left_Thumbstick_Y", "Oculus Touch (L) Thumbstick Y"), FKeyDetails::GamepadKey | FKeyDetails::Axis1D | FKeyDetails::NotBlueprintBindableKey, "OculusTouch"));
-	AddKey(FKeyDetails(EKeys::OculusTouch_Left_Thumbstick_Click, LOCTEXT("OculusTouch_Left_Thumbstick_Click", "Oculus Touch (L) Thumbstick"), FKeyDetails::GamepadKey | FKeyDetails::NotBlueprintBindableKey, "OculusTouch"));
+	AddKey(FKeyDetails(EKeys::OculusTouch_Left_Thumbstick_2D, LOCTEXT("OculusTouch_Left_Thumbstick_2D", "Oculus Touch (L) Thumbstick 2D-Axis"), FKeyDetails::GamepadKey | FKeyDetails::Axis2D | FKeyDetails::NotBlueprintBindableKey, "OculusTouch"));
+	AddKey(FKeyDetails(EKeys::OculusTouch_Left_Thumbstick_X, LOCTEXT("OculusTouch_Left_Thumbstick_X", "Oculus Touch (L) Thumbstick X-Axis"), FKeyDetails::GamepadKey | FKeyDetails::Axis1D | FKeyDetails::NotBlueprintBindableKey, "OculusTouch"));
+	AddKey(FKeyDetails(EKeys::OculusTouch_Left_Thumbstick_Y, LOCTEXT("OculusTouch_Left_Thumbstick_Y", "Oculus Touch (L) Thumbstick Y-Axis"), FKeyDetails::GamepadKey | FKeyDetails::Axis1D | FKeyDetails::NotBlueprintBindableKey, "OculusTouch"));
+	AddKey(FKeyDetails(EKeys::OculusTouch_Left_Thumbstick_Click, LOCTEXT("OculusTouch_Left_Thumbstick_Click", "Oculus Touch (L) Thumbstick Button"), FKeyDetails::GamepadKey | FKeyDetails::NotBlueprintBindableKey, "OculusTouch"));
 	AddKey(FKeyDetails(EKeys::OculusTouch_Left_Thumbstick_Touch, LOCTEXT("OculusTouch_Left_Thumbstick_Touch", "Oculus Touch (L) Thumbstick Touch"), FKeyDetails::GamepadKey | FKeyDetails::NotBlueprintBindableKey, "OculusTouch"));
 	AddKey(FKeyDetails(EKeys::OculusTouch_Left_Thumbstick_Up, LOCTEXT("OculusTouch_Left_Thumbstick_Up", "Oculus Touch (L) Thumbstick Up"), FKeyDetails::GamepadKey | FKeyDetails::ButtonAxis | FKeyDetails::NotBlueprintBindableKey, "OculusTouch"));
 	AddKey(FKeyDetails(EKeys::OculusTouch_Left_Thumbstick_Down, LOCTEXT("OculusTouch_Left_Thumbstick_Down", "Oculus Touch (L) Thumbstick Down"), FKeyDetails::GamepadKey | FKeyDetails::ButtonAxis | FKeyDetails::NotBlueprintBindableKey, "OculusTouch"));
@@ -875,9 +894,10 @@ void EKeys::Initialize()
 	AddKey(FKeyDetails(EKeys::OculusTouch_Right_Trigger_Click, LOCTEXT("OculusTouch_Right_Trigger_Click", "Oculus Touch (R) Trigger"), FKeyDetails::GamepadKey | FKeyDetails::NotBlueprintBindableKey, "OculusTouch"));
 	AddKey(FKeyDetails(EKeys::OculusTouch_Right_Trigger_Axis, LOCTEXT("OculusTouch_Right_Trigger_Axis", "Oculus Touch (R) Trigger Axis"), FKeyDetails::GamepadKey | FKeyDetails::Axis1D | FKeyDetails::NotBlueprintBindableKey, "OculusTouch"));
 	AddKey(FKeyDetails(EKeys::OculusTouch_Right_Trigger_Touch, LOCTEXT("OculusTouch_Right_Trigger_Touch", "Oculus Touch (R) Trigger Touch"), FKeyDetails::GamepadKey | FKeyDetails::NotBlueprintBindableKey, "OculusTouch"));
-	AddKey(FKeyDetails(EKeys::OculusTouch_Right_Thumbstick_X, LOCTEXT("OculusTouch_Right_Thumbstick_X", "Oculus Touch (R) Thumbstick X"), FKeyDetails::GamepadKey | FKeyDetails::Axis1D | FKeyDetails::NotBlueprintBindableKey, "OculusTouch"));
-	AddKey(FKeyDetails(EKeys::OculusTouch_Right_Thumbstick_Y, LOCTEXT("OculusTouch_Right_Thumbstick_Y", "Oculus Touch (R) Thumbstick Y"), FKeyDetails::GamepadKey | FKeyDetails::Axis1D | FKeyDetails::NotBlueprintBindableKey, "OculusTouch"));
-	AddKey(FKeyDetails(EKeys::OculusTouch_Right_Thumbstick_Click, LOCTEXT("OculusTouch_Right_Thumbstick_Click", "Oculus Touch (R) Thumbstick"), FKeyDetails::GamepadKey | FKeyDetails::NotBlueprintBindableKey, "OculusTouch"));
+	AddKey(FKeyDetails(EKeys::OculusTouch_Right_Thumbstick_2D, LOCTEXT("OculusTouch_Right_Thumbstick_2D", "Oculus Touch (R) Thumbstick 2D-Axis"), FKeyDetails::GamepadKey | FKeyDetails::Axis2D | FKeyDetails::NotBlueprintBindableKey, "OculusTouch"));
+	AddKey(FKeyDetails(EKeys::OculusTouch_Right_Thumbstick_X, LOCTEXT("OculusTouch_Right_Thumbstick_X", "Oculus Touch (R) Thumbstick X-Axis"), FKeyDetails::GamepadKey | FKeyDetails::Axis1D | FKeyDetails::NotBlueprintBindableKey, "OculusTouch"));
+	AddKey(FKeyDetails(EKeys::OculusTouch_Right_Thumbstick_Y, LOCTEXT("OculusTouch_Right_Thumbstick_Y", "Oculus Touch (R) Thumbstick Y-Axis"), FKeyDetails::GamepadKey | FKeyDetails::Axis1D | FKeyDetails::NotBlueprintBindableKey, "OculusTouch"));
+	AddKey(FKeyDetails(EKeys::OculusTouch_Right_Thumbstick_Click, LOCTEXT("OculusTouch_Right_Thumbstick_Click", "Oculus Touch (R) Thumbstick Button"), FKeyDetails::GamepadKey | FKeyDetails::NotBlueprintBindableKey, "OculusTouch"));
 	AddKey(FKeyDetails(EKeys::OculusTouch_Right_Thumbstick_Touch, LOCTEXT("OculusTouch_Right_Thumbstick_Touch", "Oculus Touch (R) Thumbstick Touch"), FKeyDetails::GamepadKey | FKeyDetails::NotBlueprintBindableKey, "OculusTouch"));
 	AddKey(FKeyDetails(EKeys::OculusTouch_Right_Thumbstick_Up, LOCTEXT("OculusTouch_Right_Thumbstick_Up", "Oculus Touch (R) Thumbstick Up"), FKeyDetails::GamepadKey | FKeyDetails::ButtonAxis | FKeyDetails::NotBlueprintBindableKey, "OculusTouch"));
 	AddKey(FKeyDetails(EKeys::OculusTouch_Right_Thumbstick_Down, LOCTEXT("OculusTouch_Right_Thumbstick_Down", "Oculus Touch (R) Thumbstick Down"), FKeyDetails::GamepadKey | FKeyDetails::ButtonAxis | FKeyDetails::NotBlueprintBindableKey, "OculusTouch"));
@@ -898,14 +918,16 @@ void EKeys::Initialize()
 	AddKey(FKeyDetails(EKeys::ValveIndex_Left_Trigger_Click, LOCTEXT("ValveIndex_Left_Trigger_Click", "Valve Index (L) Trigger"), FKeyDetails::GamepadKey | FKeyDetails::NotBlueprintBindableKey, "ValveIndex"));
 	AddKey(FKeyDetails(EKeys::ValveIndex_Left_Trigger_Axis, LOCTEXT("ValveIndex_Left_Trigger_Axis", "Valve Index (L) Trigger Axis"), FKeyDetails::GamepadKey | FKeyDetails::Axis1D | FKeyDetails::NotBlueprintBindableKey, "ValveIndex"));
 	AddKey(FKeyDetails(EKeys::ValveIndex_Left_Trigger_Touch, LOCTEXT("ValveIndex_Left_Trigger_Touch", "Valve Index (L) Trigger Touch"), FKeyDetails::GamepadKey | FKeyDetails::NotBlueprintBindableKey, "ValveIndex"));
-	AddKey(FKeyDetails(EKeys::ValveIndex_Left_Thumbstick_X, LOCTEXT("ValveIndex_Left_Thumbstick_X", "Valve Index (L) Thumbstick X"), FKeyDetails::GamepadKey | FKeyDetails::Axis1D | FKeyDetails::NotBlueprintBindableKey, "ValveIndex"));
-	AddKey(FKeyDetails(EKeys::ValveIndex_Left_Thumbstick_Y, LOCTEXT("ValveIndex_Left_Thumbstick_Y", "Valve Index (L) Thumbstick Y"), FKeyDetails::GamepadKey | FKeyDetails::Axis1D | FKeyDetails::NotBlueprintBindableKey, "ValveIndex"));
-	AddKey(FKeyDetails(EKeys::ValveIndex_Left_Thumbstick_Click, LOCTEXT("ValveIndex_Left_Thumbstick_Click", "Valve Index (L) Thumbstick"), FKeyDetails::GamepadKey | FKeyDetails::NotBlueprintBindableKey, "ValveIndex"));
+	AddKey(FKeyDetails(EKeys::ValveIndex_Left_Thumbstick_2D, LOCTEXT("ValveIndex_Left_Thumbstick_2D", "Valve Index (L) Thumbstick 2D-Axis"), FKeyDetails::GamepadKey | FKeyDetails::Axis2D | FKeyDetails::NotBlueprintBindableKey, "ValveIndex"));
+	AddKey(FKeyDetails(EKeys::ValveIndex_Left_Thumbstick_X, LOCTEXT("ValveIndex_Left_Thumbstick_X", "Valve Index (L) Thumbstick X-Axis"), FKeyDetails::GamepadKey | FKeyDetails::Axis1D | FKeyDetails::NotBlueprintBindableKey, "ValveIndex"));
+	AddKey(FKeyDetails(EKeys::ValveIndex_Left_Thumbstick_Y, LOCTEXT("ValveIndex_Left_Thumbstick_Y", "Valve Index (L) Thumbstick Y-Axis"), FKeyDetails::GamepadKey | FKeyDetails::Axis1D | FKeyDetails::NotBlueprintBindableKey, "ValveIndex"));
+	AddKey(FKeyDetails(EKeys::ValveIndex_Left_Thumbstick_Click, LOCTEXT("ValveIndex_Left_Thumbstick_Click", "Valve Index (L) Thumbstick Button"), FKeyDetails::GamepadKey | FKeyDetails::NotBlueprintBindableKey, "ValveIndex"));
 	AddKey(FKeyDetails(EKeys::ValveIndex_Left_Thumbstick_Touch, LOCTEXT("ValveIndex_Left_Thumbstick_Touch", "Valve Index (L) Thumbstick Touch"), FKeyDetails::GamepadKey | FKeyDetails::NotBlueprintBindableKey, "ValveIndex"));
 	AddKey(FKeyDetails(EKeys::ValveIndex_Left_Thumbstick_Up, LOCTEXT("ValveIndex_Left_Thumbstick_Up", "Valve Index (L) Thumbstick Up"), FKeyDetails::GamepadKey | FKeyDetails::ButtonAxis | FKeyDetails::NotBlueprintBindableKey, "ValveIndex"));
 	AddKey(FKeyDetails(EKeys::ValveIndex_Left_Thumbstick_Down, LOCTEXT("ValveIndex_Left_Thumbstick_Down", "Valve Index (L) Thumbstick Down"), FKeyDetails::GamepadKey | FKeyDetails::ButtonAxis | FKeyDetails::NotBlueprintBindableKey, "ValveIndex"));
 	AddKey(FKeyDetails(EKeys::ValveIndex_Left_Thumbstick_Left, LOCTEXT("ValveIndex_Left_Thumbstick_Left", "Valve Index (L) Thumbstick Left"), FKeyDetails::GamepadKey | FKeyDetails::ButtonAxis | FKeyDetails::NotBlueprintBindableKey, "ValveIndex"));
 	AddKey(FKeyDetails(EKeys::ValveIndex_Left_Thumbstick_Right, LOCTEXT("ValveIndex_Left_Thumbstick_Right", "Valve Index (L) Thumbstick Right"), FKeyDetails::GamepadKey | FKeyDetails::ButtonAxis | FKeyDetails::NotBlueprintBindableKey, "ValveIndex"));
+	AddKey(FKeyDetails(EKeys::ValveIndex_Left_Trackpad_2D, LOCTEXT("ValveIndex_Left_Trackpad_2D", "Valve Index (L) Trackpad 2D-Axis"), FKeyDetails::GamepadKey | FKeyDetails::Axis2D | FKeyDetails::NotBlueprintBindableKey, "ValveIndex"));
 	AddKey(FKeyDetails(EKeys::ValveIndex_Left_Trackpad_X, LOCTEXT("ValveIndex_Left_Trackpad_X", "Valve Index (L) Trackpad X"), FKeyDetails::GamepadKey | FKeyDetails::Axis1D | FKeyDetails::NotBlueprintBindableKey, "ValveIndex"));
 	AddKey(FKeyDetails(EKeys::ValveIndex_Left_Trackpad_Y, LOCTEXT("ValveIndex_Left_Trackpad_Y", "Valve Index (L) Trackpad Y"), FKeyDetails::GamepadKey | FKeyDetails::Axis1D | FKeyDetails::NotBlueprintBindableKey, "ValveIndex"));
 	AddKey(FKeyDetails(EKeys::ValveIndex_Left_Trackpad_Force, LOCTEXT("ValveIndex_Left_Trackpad_Force", "Valve Index (L) Trackpad Force"), FKeyDetails::GamepadKey | FKeyDetails::Axis1D | FKeyDetails::NotBlueprintBindableKey, "ValveIndex"));
@@ -925,14 +947,16 @@ void EKeys::Initialize()
 	AddKey(FKeyDetails(EKeys::ValveIndex_Right_Trigger_Click, LOCTEXT("ValveIndex_Right_Trigger_Click", "Valve Index (R) Trigger"), FKeyDetails::GamepadKey | FKeyDetails::NotBlueprintBindableKey, "ValveIndex"));
 	AddKey(FKeyDetails(EKeys::ValveIndex_Right_Trigger_Axis, LOCTEXT("ValveIndex_Right_Trigger_Axis", "Valve Index (R) Trigger Axis"), FKeyDetails::GamepadKey | FKeyDetails::Axis1D | FKeyDetails::NotBlueprintBindableKey, "ValveIndex"));
 	AddKey(FKeyDetails(EKeys::ValveIndex_Right_Trigger_Touch, LOCTEXT("ValveIndex_Right_Trigger_Touch", "Valve Index (R) Trigger Touch"), FKeyDetails::GamepadKey | FKeyDetails::NotBlueprintBindableKey, "ValveIndex"));
-	AddKey(FKeyDetails(EKeys::ValveIndex_Right_Thumbstick_X, LOCTEXT("ValveIndex_Right_Thumbstick_X", "Valve Index (R) Thumbstick X"), FKeyDetails::GamepadKey | FKeyDetails::Axis1D | FKeyDetails::NotBlueprintBindableKey, "ValveIndex"));
-	AddKey(FKeyDetails(EKeys::ValveIndex_Right_Thumbstick_Y, LOCTEXT("ValveIndex_Right_Thumbstick_Y", "Valve Index (R) Thumbstick Y"), FKeyDetails::GamepadKey | FKeyDetails::Axis1D | FKeyDetails::NotBlueprintBindableKey, "ValveIndex"));
-	AddKey(FKeyDetails(EKeys::ValveIndex_Right_Thumbstick_Click, LOCTEXT("ValveIndex_Right_Thumbstick_Click", "Valve Index (R) Thumbstick"), FKeyDetails::GamepadKey | FKeyDetails::NotBlueprintBindableKey, "ValveIndex"));
+	AddKey(FKeyDetails(EKeys::ValveIndex_Right_Thumbstick_2D, LOCTEXT("ValveIndex_Right_Thumbstick_2D", "Valve Index (R) Thumbstick 2D-Axis"), FKeyDetails::GamepadKey | FKeyDetails::Axis2D | FKeyDetails::NotBlueprintBindableKey, "ValveIndex"));
+	AddKey(FKeyDetails(EKeys::ValveIndex_Right_Thumbstick_X, LOCTEXT("ValveIndex_Right_Thumbstick_X", "Valve Index (R) Thumbstick X-Axis"), FKeyDetails::GamepadKey | FKeyDetails::Axis1D | FKeyDetails::NotBlueprintBindableKey, "ValveIndex"));
+	AddKey(FKeyDetails(EKeys::ValveIndex_Right_Thumbstick_Y, LOCTEXT("ValveIndex_Right_Thumbstick_Y", "Valve Index (R) Thumbstick Y-Axis"), FKeyDetails::GamepadKey | FKeyDetails::Axis1D | FKeyDetails::NotBlueprintBindableKey, "ValveIndex"));
+	AddKey(FKeyDetails(EKeys::ValveIndex_Right_Thumbstick_Click, LOCTEXT("ValveIndex_Right_Thumbstick_Click", "Valve Index (R) Thumbstick Button"), FKeyDetails::GamepadKey | FKeyDetails::NotBlueprintBindableKey, "ValveIndex"));
 	AddKey(FKeyDetails(EKeys::ValveIndex_Right_Thumbstick_Touch, LOCTEXT("ValveIndex_Right_Thumbstick_Touch", "Valve Index (R) Thumbstick Touch"), FKeyDetails::GamepadKey | FKeyDetails::NotBlueprintBindableKey, "ValveIndex"));
 	AddKey(FKeyDetails(EKeys::ValveIndex_Right_Thumbstick_Up, LOCTEXT("ValveIndex_Right_Thumbstick_Up", "Valve Index (R) Thumbstick Up"), FKeyDetails::GamepadKey | FKeyDetails::ButtonAxis | FKeyDetails::NotBlueprintBindableKey, "ValveIndex"));
 	AddKey(FKeyDetails(EKeys::ValveIndex_Right_Thumbstick_Down, LOCTEXT("ValveIndex_Right_Thumbstick_Down", "Valve Index (R) Thumbstick Down"), FKeyDetails::GamepadKey | FKeyDetails::ButtonAxis | FKeyDetails::NotBlueprintBindableKey, "ValveIndex"));
 	AddKey(FKeyDetails(EKeys::ValveIndex_Right_Thumbstick_Left, LOCTEXT("ValveIndex_Right_Thumbstick_Left", "Valve Index (R) Thumbstick Left"), FKeyDetails::GamepadKey | FKeyDetails::ButtonAxis | FKeyDetails::NotBlueprintBindableKey, "ValveIndex"));
 	AddKey(FKeyDetails(EKeys::ValveIndex_Right_Thumbstick_Right, LOCTEXT("ValveIndex_Right_Thumbstick_Right", "Valve Index (R) Thumbstick Right"), FKeyDetails::GamepadKey | FKeyDetails::ButtonAxis | FKeyDetails::NotBlueprintBindableKey, "ValveIndex"));
+	AddKey(FKeyDetails(EKeys::ValveIndex_Right_Trackpad_2D, LOCTEXT("ValveIndex_Right_Trackpad_2D", "Valve Index (R) Trackpad 2D-Axis"), FKeyDetails::GamepadKey | FKeyDetails::Axis2D | FKeyDetails::NotBlueprintBindableKey, "ValveIndex"));
 	AddKey(FKeyDetails(EKeys::ValveIndex_Right_Trackpad_X, LOCTEXT("ValveIndex_Right_Trackpad_X", "Valve Index (R) Trackpad X"), FKeyDetails::GamepadKey | FKeyDetails::Axis1D | FKeyDetails::NotBlueprintBindableKey, "ValveIndex"));
 	AddKey(FKeyDetails(EKeys::ValveIndex_Right_Trackpad_Y, LOCTEXT("ValveIndex_Right_Trackpad_Y", "Valve Index (R) Trackpad Y"), FKeyDetails::GamepadKey | FKeyDetails::Axis1D | FKeyDetails::NotBlueprintBindableKey, "ValveIndex"));
 	AddKey(FKeyDetails(EKeys::ValveIndex_Right_Trackpad_Force, LOCTEXT("ValveIndex_Right_Trackpad_Force", "Valve Index (R) Trackpad Force"), FKeyDetails::GamepadKey | FKeyDetails::Axis1D | FKeyDetails::NotBlueprintBindableKey, "ValveIndex"));
