Create actions for each individual key when not using a player mappable input config.

This will prevent OpenXR from remapping the actions in a project, but it will fix the pitfalls in the legacy system.
We should strongly encourage that developers do provide a player mappable input config for better OpenXR integration.

diff --git a/Engine/Plugins/Experimental/VirtualScouting/Source/VirtualScoutingOpenXR/Private/VirtualScoutingOpenXR.cpp b/Engine/Plugins/Experimental/VirtualScouting/Source/VirtualScoutingOpenXR/Private/VirtualScoutingOpenXR.cpp
--- a/Engine/Plugins/Experimental/VirtualScouting/Source/VirtualScoutingOpenXR/Private/VirtualScoutingOpenXR.cpp
+++ b/Engine/Plugins/Experimental/VirtualScouting/Source/VirtualScoutingOpenXR/Private/VirtualScoutingOpenXR.cpp
@@ -162,76 +162,6 @@ void FVirtualScoutingOpenXRExtension::PostCreateSession(XrSession InSession)
 }
 
 
-void FVirtualScoutingOpenXRExtension::AddActions(XrInstance InInstance, FCreateActionSetFunc CreateActionSetFunc, FCreateActionFunc CreateActionFunc)
-{
-	// Clean up any previously created action set, which also implicitly cleans up the actions belonging to it.
-	if (ActionSet != XR_NULL_HANDLE)
-	{
-		XR_ENSURE(xrDestroyActionSet(ActionSet));
-		ActionSet = XR_NULL_HANDLE;
-	}
-
-	UVREditorMode* VRMode = IVREditorModule::Get().GetVRMode();
-	if (!VRMode)
-	{
-		return;
-	}
-
-	UClass* InteractorClass = VRMode->InteractorClass.Get();
-	if (!InteractorClass)
-	{
-		return;
-	}
-
-	const UVREditorInteractor* InteractorCDO = InteractorClass->GetDefaultObject<const UVREditorInteractor>();
-	if (!InteractorCDO)
-	{
-		return;
-	}
-
-	// Create action set
-	{
-		FActionSetParams ActionSetParams;
-		ActionSetParams.Name = "ue_vp_scouting";
-		ActionSetParams.LocalizedName = LOCTEXT("ActionSetName", "Virtual Scouting");
-		ActionSetParams.Priority = 100;
-		ActionSet = CreateActionSetFunc(ActionSetParams);
-	}
-
-	XrPath LeftHand, RightHand;
-	verify(XR_SUCCEEDED(xrStringToPath(Instance, "/user/hand/left", &LeftHand)));
-	verify(XR_SUCCEEDED(xrStringToPath(Instance, "/user/hand/right", &RightHand)));
-
-	FActionParams ActionParams;
-	ActionParams.Set = ActionSet;
-	ActionParams.SubactionPaths = { LeftHand, RightHand };
-
-	for (TPair<FViewportActionKeyInput, TArray<FKey>>& ActionMappings : InteractorCDO->GetKnownActionMappings())
-	{
-		FViewportActionKeyInput& Action = ActionMappings.Key;
-		TArray<FKey>& Keys = ActionMappings.Value;
-
-		ActionParams.Type = Action.bIsAxis ? XR_ACTION_TYPE_FLOAT_INPUT : XR_ACTION_TYPE_BOOLEAN_INPUT;
-		ActionParams.Name = Action.ActionType;
-		// TODO: ActionParams.LocalizedName
-		ActionParams.SuggestedBindings = MoveTemp(Keys);
-		CreateActionFunc(ActionParams);
-	}
-}
-
-
-void FVirtualScoutingOpenXRExtension::GetActiveActionSetsForSync(TArray<XrActiveActionSet>& OutActiveSets)
-{
-	if (bIsVrEditingModeActive && ensure(ActionSet != XR_NULL_HANDLE))
-	{
-		XrActiveActionSet ActiveSet;
-		ActiveSet.actionSet = ActionSet;
-		ActiveSet.subactionPath = XR_NULL_PATH;
-		OutActiveSets.Add(ActiveSet);
-	}
-}
-
-
 void FVirtualScoutingOpenXRExtension::PostSyncActions(XrSession InSession)
 {
 	TryFulfillDeviceTypePromise();
diff --git a/Engine/Plugins/Experimental/VirtualScouting/Source/VirtualScoutingOpenXR/Private/VirtualScoutingOpenXR.h b/Engine/Plugins/Experimental/VirtualScouting/Source/VirtualScoutingOpenXR/Private/VirtualScoutingOpenXR.h
--- a/Engine/Plugins/Experimental/VirtualScouting/Source/VirtualScoutingOpenXR/Private/VirtualScoutingOpenXR.h
+++ b/Engine/Plugins/Experimental/VirtualScouting/Source/VirtualScoutingOpenXR/Private/VirtualScoutingOpenXR.h
@@ -27,8 +27,6 @@ public:
 	virtual void OnEvent(XrSession InSession, const XrEventDataBaseHeader* InHeader) override;
 	virtual void PostCreateInstance(XrInstance InInstance) override;
 	virtual void PostCreateSession(XrSession InSession) override;
-	virtual void AddActions(XrInstance InInstance, FCreateActionSetFunc CreateActionSet, FCreateActionFunc CreateAction) override;
-	virtual void GetActiveActionSetsForSync(TArray<XrActiveActionSet>& OutActiveSets) override;
 	virtual void PostSyncActions(XrSession InSession) override;
 	//~ End IOpenXRExtensionPlugin interface
 
diff --git a/Engine/Plugins/Runtime/OpenXR/Source/OpenXRHMD/Public/IOpenXRExtensionPlugin.h b/Engine/Plugins/Runtime/OpenXR/Source/OpenXRHMD/Public/IOpenXRExtensionPlugin.h
--- a/Engine/Plugins/Runtime/OpenXR/Source/OpenXRHMD/Public/IOpenXRExtensionPlugin.h
+++ b/Engine/Plugins/Runtime/OpenXR/Source/OpenXRHMD/Public/IOpenXRExtensionPlugin.h
@@ -220,41 +220,11 @@ public:
 	* Add any actions provided by the plugin to Actions.
 	* This allows a plugin to 'hard code' an action so that the plugin can use it.
 	*/
-	UE_DEPRECATED(5.1, "Use the overload taking FActionParams/FActionSetParams below instead.")
+	UE_DEPRECATED(5.1, "Use Enhanced Input through IMotionController::SetPlayerMappableInputConfig instead.")
 	virtual void AddActions(XrInstance Instance, TFunction<XrAction(XrActionType InActionType, const FName& InName, const TArray<XrPath>& InSubactionPaths)> AddAction)
 	{
 	}
 
-	struct FActionSetParams
-	{
-		FName Name;
-		FText LocalizedName;
-		uint32 Priority;
-	};
-	using FCreateActionSetFunc = TFunction< XrActionSet(const FActionSetParams&) >;
-
-	struct FActionParams
-	{
-		/** If Set == XR_NULL_HANDLE, the action is added to the default action set. */
-		XrActionSet Set = XR_NULL_HANDLE;
-
-		XrActionType Type;
-		FName Name;
-		FText LocalizedName;
-		TArray<XrPath> SubactionPaths;
-		TArray<FKey> SuggestedBindings;
-	};
-	using FCreateActionFunc = TFunction< XrAction(const FActionParams&) >;
-
-	/**
-	* Create and register actions and/or action sets belonging to the plugin.
-	* This allows a plugin to create action sets to be attached to the session,
-	* and actions which participate in interaction profile binding suggestions.
-	*/
-	virtual void AddActions(XrInstance InInstance, FCreateActionSetFunc CreateActionSet, FCreateActionFunc CreateAction)
-	{
-	}
-
 	UE_DEPRECATED(5.1, "Functionality moved to AttachActionSets().")
 	virtual void AddActionSets(TArray<XrActiveActionSet>& OutActionSets)
 	{
diff --git a/Engine/Plugins/Runtime/OpenXR/Source/OpenXRInput/Private/OpenXRInput.cpp b/Engine/Plugins/Runtime/OpenXR/Source/OpenXRInput/Private/OpenXRInput.cpp
--- a/Engine/Plugins/Runtime/OpenXR/Source/OpenXRInput/Private/OpenXRInput.cpp
+++ b/Engine/Plugins/Runtime/OpenXR/Source/OpenXRInput/Private/OpenXRInput.cpp
@@ -131,8 +131,7 @@ FOpenXRInputPlugin::FOpenXRAction::FOpenXRAction(XrActionSet InActionSet,
 	, Type(InActionType)
 	, Name(InName)
 	, Handle(XR_NULL_HANDLE)
-	, Object()
-	, KeyMap()
+	, Object(nullptr)
 {
 	char ActionName[NAME_SIZE];
 	Name.GetPlainANSIString(ActionName);
@@ -246,7 +245,6 @@ FOpenXRInputPlugin::FOpenXRInput::FOpenXRInput(FOpenXRHMD* HMD)
 	, Instance(XR_NULL_HANDLE)
 	, ControllerActionSet()
 	, ActionSets()
-	, PluginActionSets()
 	, SubactionPaths()
 	, LegacyActions()
 	, EnhancedActions()
@@ -281,13 +279,6 @@ FOpenXRInputPlugin::FOpenXRInput::FOpenXRInput(FOpenXRHMD* HMD)
 		// Generate a list of the sub-action paths so we can query the left/right hand individually
 		SubactionPaths.Add(GetPath(Instance, "/user/hand/left"));
 		SubactionPaths.Add(GetPath(Instance, "/user/hand/right"));
-
-		// Attempt to load the default input config from the OpenXR input settings.
-		UOpenXRInputSettings* InputSettings = GetMutableDefault<UOpenXRInputSettings>();
-		if (InputSettings && InputSettings->MappableInputConfig.IsValid())
-		{
-			SetPlayerMappableInputConfig((UPlayerMappableInputConfig*)InputSettings->MappableInputConfig.TryLoad());
-		}
 	}
 }
 
@@ -386,12 +377,24 @@ bool FOpenXRInputPlugin::FOpenXRInput::BuildActions(XrSession Session)
 		}
 	}
 
+	if (!MappableInputConfig)
+	{
+		// Attempt to load the default input config from the OpenXR input settings.
+		UOpenXRInputSettings* InputSettings = GetMutableDefault<UOpenXRInputSettings>();
+		if (InputSettings && InputSettings->MappableInputConfig.IsValid())
+		{
+			SetPlayerMappableInputConfig((UPlayerMappableInputConfig*)InputSettings->MappableInputConfig.TryLoad());
+		}
+	}
+
 	if (MappableInputConfig)
 	{
 		BuildEnhancedActions(Profiles);
 	}
 	else
 	{
+		UE_LOG(LogHMD, Warning, TEXT("No UPlayerMappableInputConfig provided in the OpenXR Input project settings, action bindings will not be visible to the OpenXR runtime. Action/Axis mappings from the Input configuration are deprecated for XR in favor of Enhanced Input."));
+
 		BuildLegacyActions(Profiles);
 	}
 
@@ -467,10 +470,6 @@ bool FOpenXRInputPlugin::FOpenXRInput::BuildActions(XrSession Session)
 		AttachSet.Append(PluginAttachSet);
 	}
 
-	// Bind the plugin action sets created using the deprecated API
-	for (auto&& ActionSet : PluginActionSets)
-		AttachSet.Add(ActionSet.Handle);
-
 	TArray<XrActionSet> AttachArray = AttachSet.Array();
 	XrSessionActionSetsAttachInfo SessionActionSetsAttachInfo;
 	SessionActionSetsAttachInfo.type = XR_TYPE_SESSION_ACTION_SETS_ATTACH_INFO;
@@ -486,48 +485,28 @@ void FOpenXRInputPlugin::FOpenXRInput::BuildLegacyActions(TMap<FString, FInterac
 {
 	// Create an engine action set for legacy actions
 	FOpenXRActionSet ActionSet(Instance, "ue", "Unreal Engine", 0);
-	auto InputSettings = GetMutableDefault<UInputSettings>();
-	if (InputSettings != nullptr)
+
+	TArray<FKey> Keys;
+	EKeys::GetAllKeys(Keys);
+	for (const FKey& Key : Keys)
 	{
-		TArray<FName> ActionNames;
-		InputSettings->GetActionNames(ActionNames);
-		for (const FName& ActionName : ActionNames)
+		// XR Keys are never blueprint-bindable
+		// 2D Axis keys are supported through paired keys
+		if (Key.IsBindableInBlueprints() || (Key.IsAnalog() && !Key.IsAxis1D()))
 		{
-			FOpenXRAction Action(ActionSet.Handle, XR_ACTION_TYPE_BOOLEAN_INPUT, ActionName, ActionName.ToString(), SubactionPaths);
-			TArray<FInputActionKeyMapping> Mappings;
-			InputSettings->GetActionMappingByName(ActionName, Mappings);
-
-			// If the developer didn't suggest any XR bindings for this action we won't expose it to the runtime
-			if (SuggestBindings(Profiles, Action, Mappings) > 0)
-			{
-				LegacyActions.Add(Action);
-			}
-			else
-			{
-				XR_ENSURE(xrDestroyAction(Action.Handle));
-			}
+			continue;
 		}
 
-		TArray<FName> AxisNames;
-		InputSettings->GetAxisNames(AxisNames);
-		for (const FName& AxisName : AxisNames)
+		XrActionType Type = Key.IsAxis1D() ? XR_ACTION_TYPE_FLOAT_INPUT : XR_ACTION_TYPE_BOOLEAN_INPUT;
+		FOpenXRAction Action(ActionSet.Handle, Type, Key.GetFName(), Key.GetDisplayName().ToString(), SubactionPaths);
+		if (SuggestBindingForKey(Profiles, Action, Key))
 		{
-			FOpenXRAction Action(ActionSet.Handle, XR_ACTION_TYPE_FLOAT_INPUT, AxisName, AxisName.ToString(), SubactionPaths);
-			TArray<FInputAxisKeyMapping> Mappings;
-			InputSettings->GetAxisMappingByName(AxisName, Mappings);
-
-			// If the developer didn't suggest any XR bindings for this action we won't expose it to the runtime
-			if (SuggestBindings(Profiles, Action, Mappings) > 0)
-			{
-				LegacyActions.Add(Action);
-			}
-			else
-			{
-				XR_ENSURE(xrDestroyAction(Action.Handle));
-			}
+			LegacyActions.Add(Action);
+		}
+		else
+		{
+			XR_ENSURE(xrDestroyAction(Action.Handle));
 		}
-
-		InputSettings->ForceRebuildKeymaps();
 	}
 
 	// Query extension plugins for actions
@@ -546,45 +525,7 @@ void FOpenXRInputPlugin::FOpenXRInput::BuildLegacyActions(TMap<FString, FInterac
 		PRAGMA_ENABLE_DEPRECATION_WARNINGS;
 	}
 
-	for (IOpenXRExtensionPlugin* Plugin : OpenXRHMD->GetExtensionPlugins())
-	{
-		IOpenXRExtensionPlugin::FCreateActionSetFunc CreateActionSetFunc =
-			[this]
-			(const IOpenXRExtensionPlugin::FActionSetParams& Params) -> XrActionSet
-			{
-				FOpenXRActionSet ActionSet(Instance, Params.Name, Params.LocalizedName.ToString(), Params.Priority);
-				const XrActionSet ReturnHandle = ActionSet.Handle;
-				PluginActionSets.Emplace(MoveTemp(ActionSet));
-				return ReturnHandle;
-			};
-
-		IOpenXRExtensionPlugin::FCreateActionFunc CreateActionFunc =
-			[this, &Profiles, ActionSet]
-			(const IOpenXRExtensionPlugin::FActionParams& Params) -> XrAction
-			{
-				// TODO: Use Params.LocalizedName
-				const XrActionSet UseActionSet = Params.Set == XR_NULL_HANDLE ? ActionSet.Handle : Params.Set;
-				FOpenXRAction Action(UseActionSet, Params.Type, Params.Name, Params.LocalizedName.ToString(), Params.SubactionPaths);
-				const XrAction ReturnHandle = Action.Handle;
-				for (const FKey& Key : Params.SuggestedBindings)
-				{
-					SuggestBindingForKey(Profiles, Action, Key);
-				}
-				LegacyActions.Emplace(MoveTemp(Action));
-				return ReturnHandle;
-			};
-
-		PRAGMA_DISABLE_DEPRECATION_WARNINGS;
-		Plugin->AddActions(Instance, CreateActionSetFunc, CreateActionFunc);
-		PRAGMA_ENABLE_DEPRECATION_WARNINGS;
-	}
-
-	if (!LegacyActions.IsEmpty())
-	{
-		UE_LOG(LogHMD, Warning, TEXT("Action/Axis mappings from the Input configuration are deprecated for XR in favor of Enhanced Input."));
-
-		ActionSets.Emplace(MoveTemp(ActionSet));
-	}
+	ActionSets.Emplace(MoveTemp(ActionSet));
 }
 
 void FOpenXRInputPlugin::FOpenXRInput::BuildEnhancedActions(TMap<FString, FInteractionProfile>& Profiles)
@@ -652,7 +593,6 @@ void FOpenXRInputPlugin::FOpenXRInput::DestroyActions()
 	LegacyActions.Reset();
 	EnhancedActions.Reset();
 	ActionSets.Reset();
-	PluginActionSets.Reset();
 	Controllers.Reset();
 }
 
@@ -694,10 +634,8 @@ bool FOpenXRInputPlugin::FOpenXRInput::SuggestBindingForKey(TMap<FString, FInter
 	XrPath TopLevel = GetPath(Instance, Path);
 
 	// Map this key to the correct subaction for this profile
-	// We'll use this later to trigger the correct key
+	// We'll use this later to retrieve binding modifiers
 	TPair<XrPath, XrPath> Key(Profile->Path, TopLevel);
-	Action.KeyMap.Add(Key, InFKey.GetFName());
-
 	for (UInputTrigger* Trigger : Triggers)
 	{
 		TObjectPtr<UInputTrigger> Ptr = Trigger;
@@ -799,6 +737,7 @@ void FOpenXRInputPlugin::FOpenXRInput::Tick(float DeltaTime)
 	{
 		// If the session shut down, clean up.
 		bActionsAttached = false;
+		MappableInputConfig = nullptr;
 	}
 }
 
@@ -892,52 +831,34 @@ void FOpenXRInputPlugin::FOpenXRInput::SendControllerEvents()
 
 	XrSession Session = OpenXRHMD->GetSession();
 
-	for (XrPath Subaction : SubactionPaths)
+	IPlatformInputDeviceMapper& DeviceMapper = IPlatformInputDeviceMapper::Get();
+	for (FOpenXRAction& Action : LegacyActions)
 	{
-		XrInteractionProfileState Profile;
-		Profile.type = XR_TYPE_INTERACTION_PROFILE_STATE;
-		Profile.next = nullptr;
-		XR_ENSURE(xrGetCurrentInteractionProfile(Session, Subaction, &Profile));
+		XrActionStateGetInfo GetInfo;
+		GetInfo.type = XR_TYPE_ACTION_STATE_GET_INFO;
+		GetInfo.next = nullptr;
+		GetInfo.subactionPath = XR_NULL_PATH;
+		GetInfo.action = Action.Handle;
 
-		TSet<FName> ActivatedActions, ActivatedAxes;
-		TPair<XrPath, XrPath> Key(Profile.interactionProfile, Subaction);
-		for (FOpenXRAction& Action : LegacyActions)
+		switch (Action.Type)
 		{
-			XrActionStateGetInfo GetInfo;
-			GetInfo.type = XR_TYPE_ACTION_STATE_GET_INFO;
-			GetInfo.next = nullptr;
-			GetInfo.subactionPath = Subaction;
-			GetInfo.action = Action.Handle;
-
-			// Find the action key and check if it has already been fired this frame.
-			FName* ActionKey = Action.KeyMap.Find(Key);
-			if (!ActionKey)
-			{
-				continue;
-			}
-
-			switch (Action.Type)
-			{
 			case XR_ACTION_TYPE_BOOLEAN_INPUT:
-			if (!ActivatedActions.Contains(*ActionKey))
 			{
 				XrActionStateBoolean State;
 				State.type = XR_TYPE_ACTION_STATE_BOOLEAN;
 				State.next = nullptr;
 				XrResult Result = xrGetActionStateBoolean(Session, &GetInfo, &State);
-				IPlatformInputDeviceMapper& DeviceMapper = IPlatformInputDeviceMapper::Get();
 				if (XR_SUCCEEDED(Result) && State.changedSinceLastSync)
 				{
-					ActivatedActions.Add(*ActionKey);
 					if (State.isActive && State.currentState)
 					{
 						// TODO: Map input devices here if OpenXR would like to
-						MessageHandler->OnControllerButtonPressed(*ActionKey, DeviceMapper.GetPrimaryPlatformUser(), DeviceMapper.GetDefaultInputDevice(), /*IsRepeat =*/false);
+						MessageHandler->OnControllerButtonPressed(Action.Name, DeviceMapper.GetPrimaryPlatformUser(), DeviceMapper.GetDefaultInputDevice(), /*IsRepeat =*/false);
 					}
 					else
 					{
 						// TODO: Map input devices here if OpenXR would like to
-						MessageHandler->OnControllerButtonReleased(*ActionKey, DeviceMapper.GetPrimaryPlatformUser(), DeviceMapper.GetDefaultInputDevice(), /*IsRepeat =*/false);
+						MessageHandler->OnControllerButtonReleased(Action.Name, DeviceMapper.GetPrimaryPlatformUser(), DeviceMapper.GetDefaultInputDevice(), /*IsRepeat =*/false);
 					}
 
 					FXRTimedInputActionDelegate* const Delegate = OpenXRInputNamespace::GetTimedInputActionDelegate(Action.Name);
@@ -949,25 +870,22 @@ void FOpenXRInputPlugin::FOpenXRInput::SendControllerEvents()
 			}
 			break;
 			case XR_ACTION_TYPE_FLOAT_INPUT:
-			if (!ActivatedAxes.Contains(*ActionKey))
 			{
 				XrActionStateFloat State;
 				State.type = XR_TYPE_ACTION_STATE_FLOAT;
 				State.next = nullptr;
 				XrResult Result = xrGetActionStateFloat(Session, &GetInfo, &State);
-				IPlatformInputDeviceMapper& DeviceMapper = IPlatformInputDeviceMapper::Get();
 				if (XR_SUCCEEDED(Result) && State.changedSinceLastSync)
 				{
-					ActivatedAxes.Add(*ActionKey);
 					if (State.isActive)
 					{
 						// TODO: Map input devices here if OpenXR would like to
-						MessageHandler->OnControllerAnalog(*ActionKey, DeviceMapper.GetPrimaryPlatformUser(), DeviceMapper.GetDefaultInputDevice(), State.currentState);
+						MessageHandler->OnControllerAnalog(Action.Name, DeviceMapper.GetPrimaryPlatformUser(), DeviceMapper.GetDefaultInputDevice(), State.currentState);
 					}
 					else
 					{
 						// TODO: Map input devices here if OpenXR would like to
-						MessageHandler->OnControllerAnalog(*ActionKey, DeviceMapper.GetPrimaryPlatformUser(), DeviceMapper.GetDefaultInputDevice(), 0.0f);
+						MessageHandler->OnControllerAnalog(Action.Name, DeviceMapper.GetPrimaryPlatformUser(), DeviceMapper.GetDefaultInputDevice(), 0.0f);
 					}
 
 					FXRTimedInputActionDelegate* const Delegate = OpenXRInputNamespace::GetTimedInputActionDelegate(Action.Name);
@@ -979,11 +897,19 @@ void FOpenXRInputPlugin::FOpenXRInput::SendControllerEvents()
 			}
 			break;
 			default:
-				// Other action types are currently unsupported.
-				break;
-			}
+			// Other action types are currently unsupported.
+			break;
 		}
+	}
 
+	for (XrPath Subaction : SubactionPaths)
+	{
+		XrInteractionProfileState Profile;
+		Profile.type = XR_TYPE_INTERACTION_PROFILE_STATE;
+		Profile.next = nullptr;
+		XR_ENSURE(xrGetCurrentInteractionProfile(Session, Subaction, &Profile));
+
+		TPair<XrPath, XrPath> Key(Profile.interactionProfile, Subaction);
 		for (FOpenXRAction& Action : EnhancedActions)
 		{
 			const UInputAction* InputAction = Action.Object;
diff --git a/Engine/Plugins/Runtime/OpenXR/Source/OpenXRInput/Private/OpenXRInput.h b/Engine/Plugins/Runtime/OpenXR/Source/OpenXRInput/Private/OpenXRInput.h
--- a/Engine/Plugins/Runtime/OpenXR/Source/OpenXRInput/Private/OpenXRInput.h
+++ b/Engine/Plugins/Runtime/OpenXR/Source/OpenXRInput/Private/OpenXRInput.h
@@ -36,8 +36,8 @@ public:
 		FName			Name;
 		XrAction		Handle;
 
+		// Enhanced Input
 		TObjectPtr<const UInputAction> Object;
-		TMap<TPair<XrPath, XrPath>, FName> KeyMap;
 		TMultiMap<TPair<XrPath, XrPath>, TObjectPtr<UInputTrigger>> Triggers;
 		TMultiMap<TPair<XrPath, XrPath>, TObjectPtr<UInputModifier>> Modifiers;
 
@@ -140,7 +140,6 @@ public:
 
 		TUniquePtr<FOpenXRActionSet> ControllerActionSet;
 		TArray<FOpenXRActionSet> ActionSets;
-		TArray<FOpenXRActionSet> PluginActionSets;
 		TArray<XrPath> SubactionPaths;
 		TArray<FOpenXRAction> LegacyActions, EnhancedActions;
 		TMap<EControllerHand, FOpenXRController> Controllers;
diff --git a/Engine/Source/Runtime/InputCore/Private/InputCoreTypes.cpp b/Engine/Source/Runtime/InputCore/Private/InputCoreTypes.cpp
--- a/Engine/Source/Runtime/InputCore/Private/InputCoreTypes.cpp
+++ b/Engine/Source/Runtime/InputCore/Private/InputCoreTypes.cpp
@@ -792,9 +792,9 @@ void EKeys::Initialize()
 	AddKey(FKeyDetails(EKeys::Vive_Left_Menu_Click, LOCTEXT("Vive_Left_Menu_Click", "Vive (L) Menu"), FKeyDetails::GamepadKey | FKeyDetails::NotBlueprintBindableKey, "Vive"));
 	AddKey(FKeyDetails(EKeys::Vive_Left_Trigger_Click, LOCTEXT("Vive_Left_Trigger_Click", "Vive (L) Trigger"), FKeyDetails::GamepadKey | FKeyDetails::NotBlueprintBindableKey, "Vive"));
 	AddKey(FKeyDetails(EKeys::Vive_Left_Trigger_Axis, LOCTEXT("Vive_Left_Trigger_Axis", "Vive (L) Trigger Axis"), FKeyDetails::GamepadKey | FKeyDetails::Axis1D | FKeyDetails::NotBlueprintBindableKey, "Vive"));
-	AddKey(FKeyDetails(EKeys::Vive_Left_Trackpad_2D, LOCTEXT("Vive_Left_Trackpad_2D", "Vive (L) Trackpad 2D-Axis"), FKeyDetails::GamepadKey | FKeyDetails::Axis2D | FKeyDetails::NotBlueprintBindableKey, "Vive"));
 	AddKey(FKeyDetails(EKeys::Vive_Left_Trackpad_X, LOCTEXT("Vive_Left_Trackpad_X", "Vive (L) Trackpad X-Axis"), FKeyDetails::GamepadKey | FKeyDetails::Axis1D | FKeyDetails::NotBlueprintBindableKey, "Vive"));
 	AddKey(FKeyDetails(EKeys::Vive_Left_Trackpad_Y, LOCTEXT("Vive_Left_Trackpad_Y", "Vive (L) Trackpad Y-Axis"), FKeyDetails::GamepadKey | FKeyDetails::Axis1D | FKeyDetails::NotBlueprintBindableKey, "Vive"));
+	AddPairedKey(FKeyDetails(EKeys::Vive_Left_Trackpad_2D, LOCTEXT("Vive_Left_Trackpad_2D", "Vive (L) Trackpad 2D-Axis"), FKeyDetails::GamepadKey | FKeyDetails::Axis2D | FKeyDetails::NotBlueprintBindableKey, "Vive"), EKeys::Vive_Left_Trackpad_X, EKeys::Vive_Left_Trackpad_Y);
 	AddKey(FKeyDetails(EKeys::Vive_Left_Trackpad_Click, LOCTEXT("Vive_Left_Trackpad_Click", "Vive (L) Trackpad Button"), FKeyDetails::GamepadKey | FKeyDetails::NotBlueprintBindableKey, "Vive"));
 	AddKey(FKeyDetails(EKeys::Vive_Left_Trackpad_Touch, LOCTEXT("Vive_Left_Trackpad_Touch", "Vive (L) Trackpad Touch"), FKeyDetails::GamepadKey | FKeyDetails::NotBlueprintBindableKey, "Vive"));
 	AddKey(FKeyDetails(EKeys::Vive_Left_Trackpad_Up, LOCTEXT("Vive_Left_Trackpad_Up", "Vive (L) Trackpad Up"), FKeyDetails::GamepadKey | FKeyDetails::NotBlueprintBindableKey, "Vive"));
@@ -806,9 +806,9 @@ void EKeys::Initialize()
 	AddKey(FKeyDetails(EKeys::Vive_Right_Menu_Click, LOCTEXT("Vive_Right_Menu_Click", "Vive (R) Menu"), FKeyDetails::GamepadKey | FKeyDetails::NotBlueprintBindableKey, "Vive"));
 	AddKey(FKeyDetails(EKeys::Vive_Right_Trigger_Click, LOCTEXT("Vive_Right_Trigger_Click", "Vive (R) Trigger"), FKeyDetails::GamepadKey | FKeyDetails::NotBlueprintBindableKey, "Vive"));
 	AddKey(FKeyDetails(EKeys::Vive_Right_Trigger_Axis, LOCTEXT("Vive_Right_Trigger_Axis", "Vive (R) Trigger Axis"), FKeyDetails::GamepadKey | FKeyDetails::Axis1D | FKeyDetails::NotBlueprintBindableKey, "Vive"));
-	AddKey(FKeyDetails(EKeys::Vive_Right_Trackpad_2D, LOCTEXT("Vive_Right_Trackpad_2D", "Vive (R) Trackpad 2D-Axis"), FKeyDetails::GamepadKey | FKeyDetails::Axis2D | FKeyDetails::NotBlueprintBindableKey, "Vive"));
 	AddKey(FKeyDetails(EKeys::Vive_Right_Trackpad_X, LOCTEXT("Vive_Right_Trackpad_X", "Vive (R) Trackpad X-Axis"), FKeyDetails::GamepadKey | FKeyDetails::Axis1D | FKeyDetails::NotBlueprintBindableKey, "Vive"));
 	AddKey(FKeyDetails(EKeys::Vive_Right_Trackpad_Y, LOCTEXT("Vive_Right_Trackpad_Y", "Vive (R) Trackpad Y-Axis"), FKeyDetails::GamepadKey | FKeyDetails::Axis1D | FKeyDetails::NotBlueprintBindableKey, "Vive"));
+	AddPairedKey(FKeyDetails(EKeys::Vive_Right_Trackpad_2D, LOCTEXT("Vive_Right_Trackpad_2D", "Vive (R) Trackpad 2D-Axis"), FKeyDetails::GamepadKey | FKeyDetails::Axis2D | FKeyDetails::NotBlueprintBindableKey, "Vive"), EKeys::Vive_Right_Trackpad_X, EKeys::Vive_Right_Trackpad_Y);
 	AddKey(FKeyDetails(EKeys::Vive_Right_Trackpad_Click, LOCTEXT("Vive_Right_Trackpad_Click", "Vive (R) Trackpad Button"), FKeyDetails::GamepadKey | FKeyDetails::NotBlueprintBindableKey, "Vive"));
 	AddKey(FKeyDetails(EKeys::Vive_Right_Trackpad_Touch, LOCTEXT("Vive_Right_Trackpad_Touch", "Vive (R) Trackpad Touch"), FKeyDetails::GamepadKey | FKeyDetails::NotBlueprintBindableKey, "Vive"));
 	AddKey(FKeyDetails(EKeys::Vive_Right_Trackpad_Up, LOCTEXT("Vive_Right_Trackpad_Up", "Vive (R) Trackpad Up"), FKeyDetails::GamepadKey | FKeyDetails::NotBlueprintBindableKey, "Vive"));
@@ -823,17 +823,17 @@ void EKeys::Initialize()
 	AddKey(FKeyDetails(EKeys::MixedReality_Left_Grip_Click, LOCTEXT("MixedReality_Left_Grip_Click", "Mixed Reality (L) Grip"), FKeyDetails::GamepadKey | FKeyDetails::NotBlueprintBindableKey, "MixedReality"));
 	AddKey(FKeyDetails(EKeys::MixedReality_Left_Trigger_Click, LOCTEXT("MixedReality_Left_Trigger_Click", "Mixed Reality (L) Trigger"), FKeyDetails::GamepadKey | FKeyDetails::NotBlueprintBindableKey, "MixedReality"));
 	AddKey(FKeyDetails(EKeys::MixedReality_Left_Trigger_Axis, LOCTEXT("MixedReality_Left_Trigger_Axis", "Mixed Reality (L) Trigger Axis"), FKeyDetails::GamepadKey | FKeyDetails::Axis1D | FKeyDetails::NotBlueprintBindableKey, "MixedReality"));
-	AddKey(FKeyDetails(EKeys::MixedReality_Left_Thumbstick_2D, LOCTEXT("MixedReality_Left_Thumbstick_2D", "Mixed Reality (L) Thumbstick 2D-Axis"), FKeyDetails::GamepadKey | FKeyDetails::Axis2D | FKeyDetails::NotBlueprintBindableKey, "MixedReality"));
 	AddKey(FKeyDetails(EKeys::MixedReality_Left_Thumbstick_X, LOCTEXT("MixedReality_Left_Thumbstick_X", "Mixed Reality (L) Thumbstick X-Axis"), FKeyDetails::GamepadKey | FKeyDetails::Axis1D | FKeyDetails::NotBlueprintBindableKey, "MixedReality"));
 	AddKey(FKeyDetails(EKeys::MixedReality_Left_Thumbstick_Y, LOCTEXT("MixedReality_Left_Thumbstick_Y", "Mixed Reality (L) Thumbstick Y-Axis"), FKeyDetails::GamepadKey | FKeyDetails::Axis1D | FKeyDetails::NotBlueprintBindableKey, "MixedReality"));
+	AddPairedKey(FKeyDetails(EKeys::MixedReality_Left_Thumbstick_2D, LOCTEXT("MixedReality_Left_Thumbstick_2D", "Mixed Reality (L) Thumbstick 2D-Axis"), FKeyDetails::GamepadKey | FKeyDetails::Axis2D | FKeyDetails::NotBlueprintBindableKey, "MixedReality"), EKeys::MixedReality_Left_Thumbstick_X, EKeys::MixedReality_Left_Thumbstick_Y);
 	AddKey(FKeyDetails(EKeys::MixedReality_Left_Thumbstick_Click, LOCTEXT("MixedReality_Left_Thumbstick_Click", "Mixed Reality (L) Thumbstick Button"), FKeyDetails::GamepadKey | FKeyDetails::NotBlueprintBindableKey, "MixedReality"));
 	AddKey(FKeyDetails(EKeys::MixedReality_Left_Thumbstick_Up, LOCTEXT("MixedReality_Left_Thumbstick_Up", "Mixed Reality (L) Thumbstick Up"), FKeyDetails::GamepadKey | FKeyDetails::ButtonAxis | FKeyDetails::NotBlueprintBindableKey, "MixedReality"));
 	AddKey(FKeyDetails(EKeys::MixedReality_Left_Thumbstick_Down, LOCTEXT("MixedReality_Left_Thumbstick_Down", "Mixed Reality (L) Thumbstick Down"), FKeyDetails::GamepadKey | FKeyDetails::ButtonAxis | FKeyDetails::NotBlueprintBindableKey, "MixedReality"));
 	AddKey(FKeyDetails(EKeys::MixedReality_Left_Thumbstick_Left, LOCTEXT("MixedReality_Left_Thumbstick_Left", "Mixed Reality (L) Thumbstick Left"), FKeyDetails::GamepadKey | FKeyDetails::ButtonAxis | FKeyDetails::NotBlueprintBindableKey, "MixedReality"));
 	AddKey(FKeyDetails(EKeys::MixedReality_Left_Thumbstick_Right, LOCTEXT("MixedReality_Left_Thumbstick_Right", "Mixed Reality (L) Thumbstick Right"), FKeyDetails::GamepadKey | FKeyDetails::ButtonAxis | FKeyDetails::NotBlueprintBindableKey, "MixedReality"));
-	AddKey(FKeyDetails(EKeys::MixedReality_Left_Trackpad_2D, LOCTEXT("MixedReality_Left_Trackpad_2D", "Mixed Reality (L) Trackpad 2D-Axis"), FKeyDetails::GamepadKey | FKeyDetails::Axis2D | FKeyDetails::NotBlueprintBindableKey, "MixedReality"));
 	AddKey(FKeyDetails(EKeys::MixedReality_Left_Trackpad_X, LOCTEXT("MixedReality_Left_Trackpad_X", "Mixed Reality (L) Trackpad X"), FKeyDetails::GamepadKey | FKeyDetails::Axis1D | FKeyDetails::NotBlueprintBindableKey, "MixedReality"));
 	AddKey(FKeyDetails(EKeys::MixedReality_Left_Trackpad_Y, LOCTEXT("MixedReality_Left_Trackpad_Y", "Mixed Reality (L) Trackpad Y"), FKeyDetails::GamepadKey | FKeyDetails::Axis1D | FKeyDetails::NotBlueprintBindableKey, "MixedReality"));
+	AddPairedKey(FKeyDetails(EKeys::MixedReality_Left_Trackpad_2D, LOCTEXT("MixedReality_Left_Trackpad_2D", "Mixed Reality (L) Trackpad 2D-Axis"), FKeyDetails::GamepadKey | FKeyDetails::Axis2D | FKeyDetails::NotBlueprintBindableKey, "MixedReality"), EKeys::MixedReality_Left_Trackpad_X, EKeys::MixedReality_Left_Trackpad_Y);
 	AddKey(FKeyDetails(EKeys::MixedReality_Left_Trackpad_Click, LOCTEXT("MixedReality_Left_Trackpad_Click", "Mixed Reality (L) Trackpad Button"), FKeyDetails::GamepadKey | FKeyDetails::NotBlueprintBindableKey, "MixedReality"));
 	AddKey(FKeyDetails(EKeys::MixedReality_Left_Trackpad_Touch, LOCTEXT("MixedReality_Left_Trackpad_Touch", "Mixed Reality (L) Trackpad Touch"), FKeyDetails::GamepadKey | FKeyDetails::NotBlueprintBindableKey, "MixedReality"));
 	AddKey(FKeyDetails(EKeys::MixedReality_Left_Trackpad_Up, LOCTEXT("MixedReality_Left_Trackpad_Up", "Mixed Reality (L) Trackpad Up"), FKeyDetails::GamepadKey | FKeyDetails::NotBlueprintBindableKey, "MixedReality"));
@@ -844,17 +844,17 @@ void EKeys::Initialize()
 	AddKey(FKeyDetails(EKeys::MixedReality_Right_Grip_Click, LOCTEXT("MixedReality_Right_Grip_Click", "Mixed Reality (R) Grip"), FKeyDetails::GamepadKey | FKeyDetails::NotBlueprintBindableKey, "MixedReality"));
 	AddKey(FKeyDetails(EKeys::MixedReality_Right_Trigger_Click, LOCTEXT("MixedReality_Right_Trigger_Click", "Mixed Reality (R) Trigger"), FKeyDetails::GamepadKey | FKeyDetails::NotBlueprintBindableKey, "MixedReality"));
 	AddKey(FKeyDetails(EKeys::MixedReality_Right_Trigger_Axis, LOCTEXT("MixedReality_Right_Trigger_Axis", "Mixed Reality (R) Trigger Axis"), FKeyDetails::GamepadKey | FKeyDetails::Axis1D | FKeyDetails::NotBlueprintBindableKey, "MixedReality"));
-	AddKey(FKeyDetails(EKeys::MixedReality_Right_Thumbstick_2D, LOCTEXT("MixedReality_Right_Thumbstick_2D", "Mixed Reality (R) Thumbstick 2D-Axis"), FKeyDetails::GamepadKey | FKeyDetails::Axis2D | FKeyDetails::NotBlueprintBindableKey, "MixedReality"));
 	AddKey(FKeyDetails(EKeys::MixedReality_Right_Thumbstick_X, LOCTEXT("MixedReality_Right_Thumbstick_X", "Mixed Reality (R) Thumbstick X-Axis"), FKeyDetails::GamepadKey | FKeyDetails::Axis1D | FKeyDetails::NotBlueprintBindableKey, "MixedReality"));
 	AddKey(FKeyDetails(EKeys::MixedReality_Right_Thumbstick_Y, LOCTEXT("MixedReality_Right_Thumbstick_Y", "Mixed Reality (R) Thumbstick Y-Axis"), FKeyDetails::GamepadKey | FKeyDetails::Axis1D | FKeyDetails::NotBlueprintBindableKey, "MixedReality"));
+	AddPairedKey(FKeyDetails(EKeys::MixedReality_Right_Thumbstick_2D, LOCTEXT("MixedReality_Right_Thumbstick_2D", "Mixed Reality (R) Thumbstick 2D-Axis"), FKeyDetails::GamepadKey | FKeyDetails::Axis2D | FKeyDetails::NotBlueprintBindableKey, "MixedReality"), EKeys::MixedReality_Right_Thumbstick_X, EKeys::MixedReality_Right_Thumbstick_Y);
 	AddKey(FKeyDetails(EKeys::MixedReality_Right_Thumbstick_Click, LOCTEXT("MixedReality_Right_Thumbstick_Click", "Mixed Reality (R) Thumbstick Button"), FKeyDetails::GamepadKey | FKeyDetails::NotBlueprintBindableKey, "MixedReality"));
 	AddKey(FKeyDetails(EKeys::MixedReality_Right_Thumbstick_Up, LOCTEXT("MixedReality_Right_Thumbstick_Up", "Mixed Reality (R) Thumbstick Up"), FKeyDetails::GamepadKey | FKeyDetails::ButtonAxis | FKeyDetails::NotBlueprintBindableKey, "MixedReality"));
 	AddKey(FKeyDetails(EKeys::MixedReality_Right_Thumbstick_Down, LOCTEXT("MixedReality_Right_Thumbstick_Down", "Mixed Reality (R) Thumbstick Down"), FKeyDetails::GamepadKey | FKeyDetails::ButtonAxis | FKeyDetails::NotBlueprintBindableKey, "MixedReality"));
 	AddKey(FKeyDetails(EKeys::MixedReality_Right_Thumbstick_Left, LOCTEXT("MixedReality_Right_Thumbstick_Left", "Mixed Reality (R) Thumbstick Left"), FKeyDetails::GamepadKey | FKeyDetails::ButtonAxis | FKeyDetails::NotBlueprintBindableKey, "MixedReality"));
 	AddKey(FKeyDetails(EKeys::MixedReality_Right_Thumbstick_Right, LOCTEXT("MixedReality_Right_Thumbstick_Right", "Mixed Reality (R) Thumbstick Right"), FKeyDetails::GamepadKey | FKeyDetails::ButtonAxis | FKeyDetails::NotBlueprintBindableKey, "MixedReality"));
-	AddKey(FKeyDetails(EKeys::MixedReality_Right_Trackpad_2D, LOCTEXT("MixedReality_Right_Trackpad_2D", "Mixed Reality (R) Trackpad 2D-Axis"), FKeyDetails::GamepadKey | FKeyDetails::Axis2D | FKeyDetails::NotBlueprintBindableKey, "MixedReality"));
 	AddKey(FKeyDetails(EKeys::MixedReality_Right_Trackpad_X, LOCTEXT("MixedReality_Right_Trackpad_X", "Mixed Reality (R) Trackpad X"), FKeyDetails::GamepadKey | FKeyDetails::Axis1D | FKeyDetails::NotBlueprintBindableKey, "MixedReality"));
 	AddKey(FKeyDetails(EKeys::MixedReality_Right_Trackpad_Y, LOCTEXT("MixedReality_Right_Trackpad_Y", "Mixed Reality (R) Trackpad Y"), FKeyDetails::GamepadKey | FKeyDetails::Axis1D | FKeyDetails::NotBlueprintBindableKey, "MixedReality"));
+	AddPairedKey(FKeyDetails(EKeys::MixedReality_Right_Trackpad_2D, LOCTEXT("MixedReality_Right_Trackpad_2D", "Mixed Reality (R) Trackpad 2D-Axis"), FKeyDetails::GamepadKey | FKeyDetails::Axis2D | FKeyDetails::NotBlueprintBindableKey, "MixedReality"), EKeys::MixedReality_Right_Trackpad_X, EKeys::MixedReality_Right_Trackpad_Y);
 	AddKey(FKeyDetails(EKeys::MixedReality_Right_Trackpad_Click, LOCTEXT("MixedReality_Right_Trackpad_Click", "Mixed Reality (R) Trackpad Button"), FKeyDetails::GamepadKey | FKeyDetails::NotBlueprintBindableKey, "MixedReality"));
 	AddKey(FKeyDetails(EKeys::MixedReality_Right_Trackpad_Touch, LOCTEXT("MixedReality_Right_Trackpad_Touch", "Mixed Reality (R) Trackpad Touch"), FKeyDetails::GamepadKey | FKeyDetails::NotBlueprintBindableKey, "MixedReality"));
 	AddKey(FKeyDetails(EKeys::MixedReality_Right_Trackpad_Up, LOCTEXT("MixedReality_Right_Trackpad_Up", "Mixed Reality (R) Trackpad Up"), FKeyDetails::GamepadKey | FKeyDetails::NotBlueprintBindableKey, "MixedReality"));
@@ -875,9 +875,9 @@ void EKeys::Initialize()
 	AddKey(FKeyDetails(EKeys::OculusTouch_Left_Trigger_Click, LOCTEXT("OculusTouch_Left_Trigger_Click", "Oculus Touch (L) Trigger"), FKeyDetails::GamepadKey | FKeyDetails::NotBlueprintBindableKey, "OculusTouch"));
 	AddKey(FKeyDetails(EKeys::OculusTouch_Left_Trigger_Axis, LOCTEXT("OculusTouch_Left_Trigger_Axis", "Oculus Touch (L) Trigger Axis"), FKeyDetails::GamepadKey | FKeyDetails::Axis1D | FKeyDetails::NotBlueprintBindableKey, "OculusTouch"));
 	AddKey(FKeyDetails(EKeys::OculusTouch_Left_Trigger_Touch, LOCTEXT("OculusTouch_Left_Trigger_Touch", "Oculus Touch (L) Trigger Touch"), FKeyDetails::GamepadKey | FKeyDetails::NotBlueprintBindableKey, "OculusTouch"));
-	AddKey(FKeyDetails(EKeys::OculusTouch_Left_Thumbstick_2D, LOCTEXT("OculusTouch_Left_Thumbstick_2D", "Oculus Touch (L) Thumbstick 2D-Axis"), FKeyDetails::GamepadKey | FKeyDetails::Axis2D | FKeyDetails::NotBlueprintBindableKey, "OculusTouch"));
 	AddKey(FKeyDetails(EKeys::OculusTouch_Left_Thumbstick_X, LOCTEXT("OculusTouch_Left_Thumbstick_X", "Oculus Touch (L) Thumbstick X-Axis"), FKeyDetails::GamepadKey | FKeyDetails::Axis1D | FKeyDetails::NotBlueprintBindableKey, "OculusTouch"));
 	AddKey(FKeyDetails(EKeys::OculusTouch_Left_Thumbstick_Y, LOCTEXT("OculusTouch_Left_Thumbstick_Y", "Oculus Touch (L) Thumbstick Y-Axis"), FKeyDetails::GamepadKey | FKeyDetails::Axis1D | FKeyDetails::NotBlueprintBindableKey, "OculusTouch"));
+	AddPairedKey(FKeyDetails(EKeys::OculusTouch_Left_Thumbstick_2D, LOCTEXT("OculusTouch_Left_Thumbstick_2D", "Oculus Touch (L) Thumbstick 2D-Axis"), FKeyDetails::GamepadKey | FKeyDetails::Axis2D | FKeyDetails::NotBlueprintBindableKey, "OculusTouch"), EKeys::OculusTouch_Left_Thumbstick_X, EKeys::OculusTouch_Left_Thumbstick_Y);
 	AddKey(FKeyDetails(EKeys::OculusTouch_Left_Thumbstick_Click, LOCTEXT("OculusTouch_Left_Thumbstick_Click", "Oculus Touch (L) Thumbstick Button"), FKeyDetails::GamepadKey | FKeyDetails::NotBlueprintBindableKey, "OculusTouch"));
 	AddKey(FKeyDetails(EKeys::OculusTouch_Left_Thumbstick_Touch, LOCTEXT("OculusTouch_Left_Thumbstick_Touch", "Oculus Touch (L) Thumbstick Touch"), FKeyDetails::GamepadKey | FKeyDetails::NotBlueprintBindableKey, "OculusTouch"));
 	AddKey(FKeyDetails(EKeys::OculusTouch_Left_Thumbstick_Up, LOCTEXT("OculusTouch_Left_Thumbstick_Up", "Oculus Touch (L) Thumbstick Up"), FKeyDetails::GamepadKey | FKeyDetails::ButtonAxis | FKeyDetails::NotBlueprintBindableKey, "OculusTouch"));
@@ -894,9 +894,9 @@ void EKeys::Initialize()
 	AddKey(FKeyDetails(EKeys::OculusTouch_Right_Trigger_Click, LOCTEXT("OculusTouch_Right_Trigger_Click", "Oculus Touch (R) Trigger"), FKeyDetails::GamepadKey | FKeyDetails::NotBlueprintBindableKey, "OculusTouch"));
 	AddKey(FKeyDetails(EKeys::OculusTouch_Right_Trigger_Axis, LOCTEXT("OculusTouch_Right_Trigger_Axis", "Oculus Touch (R) Trigger Axis"), FKeyDetails::GamepadKey | FKeyDetails::Axis1D | FKeyDetails::NotBlueprintBindableKey, "OculusTouch"));
 	AddKey(FKeyDetails(EKeys::OculusTouch_Right_Trigger_Touch, LOCTEXT("OculusTouch_Right_Trigger_Touch", "Oculus Touch (R) Trigger Touch"), FKeyDetails::GamepadKey | FKeyDetails::NotBlueprintBindableKey, "OculusTouch"));
-	AddKey(FKeyDetails(EKeys::OculusTouch_Right_Thumbstick_2D, LOCTEXT("OculusTouch_Right_Thumbstick_2D", "Oculus Touch (R) Thumbstick 2D-Axis"), FKeyDetails::GamepadKey | FKeyDetails::Axis2D | FKeyDetails::NotBlueprintBindableKey, "OculusTouch"));
 	AddKey(FKeyDetails(EKeys::OculusTouch_Right_Thumbstick_X, LOCTEXT("OculusTouch_Right_Thumbstick_X", "Oculus Touch (R) Thumbstick X-Axis"), FKeyDetails::GamepadKey | FKeyDetails::Axis1D | FKeyDetails::NotBlueprintBindableKey, "OculusTouch"));
 	AddKey(FKeyDetails(EKeys::OculusTouch_Right_Thumbstick_Y, LOCTEXT("OculusTouch_Right_Thumbstick_Y", "Oculus Touch (R) Thumbstick Y-Axis"), FKeyDetails::GamepadKey | FKeyDetails::Axis1D | FKeyDetails::NotBlueprintBindableKey, "OculusTouch"));
+	AddPairedKey(FKeyDetails(EKeys::OculusTouch_Right_Thumbstick_2D, LOCTEXT("OculusTouch_Right_Thumbstick_2D", "Oculus Touch (R) Thumbstick 2D-Axis"), FKeyDetails::GamepadKey | FKeyDetails::Axis2D | FKeyDetails::NotBlueprintBindableKey, "OculusTouch"), EKeys::OculusTouch_Right_Thumbstick_X, EKeys::OculusTouch_Right_Thumbstick_Y);
 	AddKey(FKeyDetails(EKeys::OculusTouch_Right_Thumbstick_Click, LOCTEXT("OculusTouch_Right_Thumbstick_Click", "Oculus Touch (R) Thumbstick Button"), FKeyDetails::GamepadKey | FKeyDetails::NotBlueprintBindableKey, "OculusTouch"));
 	AddKey(FKeyDetails(EKeys::OculusTouch_Right_Thumbstick_Touch, LOCTEXT("OculusTouch_Right_Thumbstick_Touch", "Oculus Touch (R) Thumbstick Touch"), FKeyDetails::GamepadKey | FKeyDetails::NotBlueprintBindableKey, "OculusTouch"));
 	AddKey(FKeyDetails(EKeys::OculusTouch_Right_Thumbstick_Up, LOCTEXT("OculusTouch_Right_Thumbstick_Up", "Oculus Touch (R) Thumbstick Up"), FKeyDetails::GamepadKey | FKeyDetails::ButtonAxis | FKeyDetails::NotBlueprintBindableKey, "OculusTouch"));
@@ -918,18 +918,18 @@ void EKeys::Initialize()
 	AddKey(FKeyDetails(EKeys::ValveIndex_Left_Trigger_Click, LOCTEXT("ValveIndex_Left_Trigger_Click", "Valve Index (L) Trigger"), FKeyDetails::GamepadKey | FKeyDetails::NotBlueprintBindableKey, "ValveIndex"));
 	AddKey(FKeyDetails(EKeys::ValveIndex_Left_Trigger_Axis, LOCTEXT("ValveIndex_Left_Trigger_Axis", "Valve Index (L) Trigger Axis"), FKeyDetails::GamepadKey | FKeyDetails::Axis1D | FKeyDetails::NotBlueprintBindableKey, "ValveIndex"));
 	AddKey(FKeyDetails(EKeys::ValveIndex_Left_Trigger_Touch, LOCTEXT("ValveIndex_Left_Trigger_Touch", "Valve Index (L) Trigger Touch"), FKeyDetails::GamepadKey | FKeyDetails::NotBlueprintBindableKey, "ValveIndex"));
-	AddKey(FKeyDetails(EKeys::ValveIndex_Left_Thumbstick_2D, LOCTEXT("ValveIndex_Left_Thumbstick_2D", "Valve Index (L) Thumbstick 2D-Axis"), FKeyDetails::GamepadKey | FKeyDetails::Axis2D | FKeyDetails::NotBlueprintBindableKey, "ValveIndex"));
 	AddKey(FKeyDetails(EKeys::ValveIndex_Left_Thumbstick_X, LOCTEXT("ValveIndex_Left_Thumbstick_X", "Valve Index (L) Thumbstick X-Axis"), FKeyDetails::GamepadKey | FKeyDetails::Axis1D | FKeyDetails::NotBlueprintBindableKey, "ValveIndex"));
 	AddKey(FKeyDetails(EKeys::ValveIndex_Left_Thumbstick_Y, LOCTEXT("ValveIndex_Left_Thumbstick_Y", "Valve Index (L) Thumbstick Y-Axis"), FKeyDetails::GamepadKey | FKeyDetails::Axis1D | FKeyDetails::NotBlueprintBindableKey, "ValveIndex"));
+	AddPairedKey(FKeyDetails(EKeys::ValveIndex_Left_Thumbstick_2D, LOCTEXT("ValveIndex_Left_Thumbstick_2D", "Valve Index (L) Thumbstick 2D-Axis"), FKeyDetails::GamepadKey | FKeyDetails::Axis2D | FKeyDetails::NotBlueprintBindableKey, "ValveIndex"), EKeys::ValveIndex_Left_Thumbstick_X, EKeys::ValveIndex_Left_Thumbstick_Y);
 	AddKey(FKeyDetails(EKeys::ValveIndex_Left_Thumbstick_Click, LOCTEXT("ValveIndex_Left_Thumbstick_Click", "Valve Index (L) Thumbstick Button"), FKeyDetails::GamepadKey | FKeyDetails::NotBlueprintBindableKey, "ValveIndex"));
 	AddKey(FKeyDetails(EKeys::ValveIndex_Left_Thumbstick_Touch, LOCTEXT("ValveIndex_Left_Thumbstick_Touch", "Valve Index (L) Thumbstick Touch"), FKeyDetails::GamepadKey | FKeyDetails::NotBlueprintBindableKey, "ValveIndex"));
 	AddKey(FKeyDetails(EKeys::ValveIndex_Left_Thumbstick_Up, LOCTEXT("ValveIndex_Left_Thumbstick_Up", "Valve Index (L) Thumbstick Up"), FKeyDetails::GamepadKey | FKeyDetails::ButtonAxis | FKeyDetails::NotBlueprintBindableKey, "ValveIndex"));
 	AddKey(FKeyDetails(EKeys::ValveIndex_Left_Thumbstick_Down, LOCTEXT("ValveIndex_Left_Thumbstick_Down", "Valve Index (L) Thumbstick Down"), FKeyDetails::GamepadKey | FKeyDetails::ButtonAxis | FKeyDetails::NotBlueprintBindableKey, "ValveIndex"));
 	AddKey(FKeyDetails(EKeys::ValveIndex_Left_Thumbstick_Left, LOCTEXT("ValveIndex_Left_Thumbstick_Left", "Valve Index (L) Thumbstick Left"), FKeyDetails::GamepadKey | FKeyDetails::ButtonAxis | FKeyDetails::NotBlueprintBindableKey, "ValveIndex"));
 	AddKey(FKeyDetails(EKeys::ValveIndex_Left_Thumbstick_Right, LOCTEXT("ValveIndex_Left_Thumbstick_Right", "Valve Index (L) Thumbstick Right"), FKeyDetails::GamepadKey | FKeyDetails::ButtonAxis | FKeyDetails::NotBlueprintBindableKey, "ValveIndex"));
-	AddKey(FKeyDetails(EKeys::ValveIndex_Left_Trackpad_2D, LOCTEXT("ValveIndex_Left_Trackpad_2D", "Valve Index (L) Trackpad 2D-Axis"), FKeyDetails::GamepadKey | FKeyDetails::Axis2D | FKeyDetails::NotBlueprintBindableKey, "ValveIndex"));
 	AddKey(FKeyDetails(EKeys::ValveIndex_Left_Trackpad_X, LOCTEXT("ValveIndex_Left_Trackpad_X", "Valve Index (L) Trackpad X"), FKeyDetails::GamepadKey | FKeyDetails::Axis1D | FKeyDetails::NotBlueprintBindableKey, "ValveIndex"));
 	AddKey(FKeyDetails(EKeys::ValveIndex_Left_Trackpad_Y, LOCTEXT("ValveIndex_Left_Trackpad_Y", "Valve Index (L) Trackpad Y"), FKeyDetails::GamepadKey | FKeyDetails::Axis1D | FKeyDetails::NotBlueprintBindableKey, "ValveIndex"));
+	AddPairedKey(FKeyDetails(EKeys::ValveIndex_Left_Trackpad_2D, LOCTEXT("ValveIndex_Left_Trackpad_2D", "Valve Index (L) Trackpad 2D-Axis"), FKeyDetails::GamepadKey | FKeyDetails::Axis2D | FKeyDetails::NotBlueprintBindableKey, "ValveIndex"), EKeys::ValveIndex_Left_Trackpad_X, EKeys::ValveIndex_Left_Trackpad_Y);
 	AddKey(FKeyDetails(EKeys::ValveIndex_Left_Trackpad_Force, LOCTEXT("ValveIndex_Left_Trackpad_Force", "Valve Index (L) Trackpad Force"), FKeyDetails::GamepadKey | FKeyDetails::Axis1D | FKeyDetails::NotBlueprintBindableKey, "ValveIndex"));
 	AddKey(FKeyDetails(EKeys::ValveIndex_Left_Trackpad_Touch, LOCTEXT("ValveIndex_Left_Trackpad_Touch", "Valve Index (L) Trackpad Touch"), FKeyDetails::GamepadKey | FKeyDetails::NotBlueprintBindableKey, "ValveIndex"));
 	AddKey(FKeyDetails(EKeys::ValveIndex_Left_Trackpad_Up, LOCTEXT("ValveIndex_Left_Trackpad_Up", "Valve Index (L) Trackpad Up"), FKeyDetails::GamepadKey | FKeyDetails::NotBlueprintBindableKey, "ValveIndex"));
@@ -947,18 +947,18 @@ void EKeys::Initialize()
 	AddKey(FKeyDetails(EKeys::ValveIndex_Right_Trigger_Click, LOCTEXT("ValveIndex_Right_Trigger_Click", "Valve Index (R) Trigger"), FKeyDetails::GamepadKey | FKeyDetails::NotBlueprintBindableKey, "ValveIndex"));
 	AddKey(FKeyDetails(EKeys::ValveIndex_Right_Trigger_Axis, LOCTEXT("ValveIndex_Right_Trigger_Axis", "Valve Index (R) Trigger Axis"), FKeyDetails::GamepadKey | FKeyDetails::Axis1D | FKeyDetails::NotBlueprintBindableKey, "ValveIndex"));
 	AddKey(FKeyDetails(EKeys::ValveIndex_Right_Trigger_Touch, LOCTEXT("ValveIndex_Right_Trigger_Touch", "Valve Index (R) Trigger Touch"), FKeyDetails::GamepadKey | FKeyDetails::NotBlueprintBindableKey, "ValveIndex"));
-	AddKey(FKeyDetails(EKeys::ValveIndex_Right_Thumbstick_2D, LOCTEXT("ValveIndex_Right_Thumbstick_2D", "Valve Index (R) Thumbstick 2D-Axis"), FKeyDetails::GamepadKey | FKeyDetails::Axis2D | FKeyDetails::NotBlueprintBindableKey, "ValveIndex"));
 	AddKey(FKeyDetails(EKeys::ValveIndex_Right_Thumbstick_X, LOCTEXT("ValveIndex_Right_Thumbstick_X", "Valve Index (R) Thumbstick X-Axis"), FKeyDetails::GamepadKey | FKeyDetails::Axis1D | FKeyDetails::NotBlueprintBindableKey, "ValveIndex"));
 	AddKey(FKeyDetails(EKeys::ValveIndex_Right_Thumbstick_Y, LOCTEXT("ValveIndex_Right_Thumbstick_Y", "Valve Index (R) Thumbstick Y-Axis"), FKeyDetails::GamepadKey | FKeyDetails::Axis1D | FKeyDetails::NotBlueprintBindableKey, "ValveIndex"));
+	AddPairedKey(FKeyDetails(EKeys::ValveIndex_Right_Thumbstick_2D, LOCTEXT("ValveIndex_Right_Thumbstick_2D", "Valve Index (R) Thumbstick 2D-Axis"), FKeyDetails::GamepadKey | FKeyDetails::Axis2D | FKeyDetails::NotBlueprintBindableKey, "ValveIndex"), EKeys::ValveIndex_Right_Thumbstick_X, EKeys::ValveIndex_Right_Thumbstick_Y);
 	AddKey(FKeyDetails(EKeys::ValveIndex_Right_Thumbstick_Click, LOCTEXT("ValveIndex_Right_Thumbstick_Click", "Valve Index (R) Thumbstick Button"), FKeyDetails::GamepadKey | FKeyDetails::NotBlueprintBindableKey, "ValveIndex"));
 	AddKey(FKeyDetails(EKeys::ValveIndex_Right_Thumbstick_Touch, LOCTEXT("ValveIndex_Right_Thumbstick_Touch", "Valve Index (R) Thumbstick Touch"), FKeyDetails::GamepadKey | FKeyDetails::NotBlueprintBindableKey, "ValveIndex"));
 	AddKey(FKeyDetails(EKeys::ValveIndex_Right_Thumbstick_Up, LOCTEXT("ValveIndex_Right_Thumbstick_Up", "Valve Index (R) Thumbstick Up"), FKeyDetails::GamepadKey | FKeyDetails::ButtonAxis | FKeyDetails::NotBlueprintBindableKey, "ValveIndex"));
 	AddKey(FKeyDetails(EKeys::ValveIndex_Right_Thumbstick_Down, LOCTEXT("ValveIndex_Right_Thumbstick_Down", "Valve Index (R) Thumbstick Down"), FKeyDetails::GamepadKey | FKeyDetails::ButtonAxis | FKeyDetails::NotBlueprintBindableKey, "ValveIndex"));
 	AddKey(FKeyDetails(EKeys::ValveIndex_Right_Thumbstick_Left, LOCTEXT("ValveIndex_Right_Thumbstick_Left", "Valve Index (R) Thumbstick Left"), FKeyDetails::GamepadKey | FKeyDetails::ButtonAxis | FKeyDetails::NotBlueprintBindableKey, "ValveIndex"));
 	AddKey(FKeyDetails(EKeys::ValveIndex_Right_Thumbstick_Right, LOCTEXT("ValveIndex_Right_Thumbstick_Right", "Valve Index (R) Thumbstick Right"), FKeyDetails::GamepadKey | FKeyDetails::ButtonAxis | FKeyDetails::NotBlueprintBindableKey, "ValveIndex"));
-	AddKey(FKeyDetails(EKeys::ValveIndex_Right_Trackpad_2D, LOCTEXT("ValveIndex_Right_Trackpad_2D", "Valve Index (R) Trackpad 2D-Axis"), FKeyDetails::GamepadKey | FKeyDetails::Axis2D | FKeyDetails::NotBlueprintBindableKey, "ValveIndex"));
 	AddKey(FKeyDetails(EKeys::ValveIndex_Right_Trackpad_X, LOCTEXT("ValveIndex_Right_Trackpad_X", "Valve Index (R) Trackpad X"), FKeyDetails::GamepadKey | FKeyDetails::Axis1D | FKeyDetails::NotBlueprintBindableKey, "ValveIndex"));
 	AddKey(FKeyDetails(EKeys::ValveIndex_Right_Trackpad_Y, LOCTEXT("ValveIndex_Right_Trackpad_Y", "Valve Index (R) Trackpad Y"), FKeyDetails::GamepadKey | FKeyDetails::Axis1D | FKeyDetails::NotBlueprintBindableKey, "ValveIndex"));
+	AddPairedKey(FKeyDetails(EKeys::ValveIndex_Right_Trackpad_2D, LOCTEXT("ValveIndex_Right_Trackpad_2D", "Valve Index (R) Trackpad 2D-Axis"), FKeyDetails::GamepadKey | FKeyDetails::Axis2D | FKeyDetails::NotBlueprintBindableKey, "ValveIndex"), EKeys::ValveIndex_Right_Trackpad_X, EKeys::ValveIndex_Right_Trackpad_Y);
 	AddKey(FKeyDetails(EKeys::ValveIndex_Right_Trackpad_Force, LOCTEXT("ValveIndex_Right_Trackpad_Force", "Valve Index (R) Trackpad Force"), FKeyDetails::GamepadKey | FKeyDetails::Axis1D | FKeyDetails::NotBlueprintBindableKey, "ValveIndex"));
 	AddKey(FKeyDetails(EKeys::ValveIndex_Right_Trackpad_Touch, LOCTEXT("ValveIndex_Right_Trackpad_Touch", "Valve Index (R) Trackpad Touch"), FKeyDetails::GamepadKey | FKeyDetails::NotBlueprintBindableKey, "ValveIndex"));
 	AddKey(FKeyDetails(EKeys::ValveIndex_Right_Trackpad_Up, LOCTEXT("ValveIndex_Right_Trackpad_Up", "Valve Index (R) Trackpad Up"), FKeyDetails::GamepadKey | FKeyDetails::NotBlueprintBindableKey, "ValveIndex"));
