OpenXR / InputCore: Deprecate XR system click/touch FKeys, which cannot be bound in OpenXR.

diff --git a/Engine/Plugins/Runtime/OpenXR/Source/OpenXRInput/Private/OpenXRInput.cpp b/Engine/Plugins/Runtime/OpenXR/Source/OpenXRInput/Private/OpenXRInput.cpp
--- a/Engine/Plugins/Runtime/OpenXR/Source/OpenXRInput/Private/OpenXRInput.cpp
+++ b/Engine/Plugins/Runtime/OpenXR/Source/OpenXRInput/Private/OpenXRInput.cpp
@@ -490,7 +490,7 @@ void FOpenXRInputPlugin::FOpenXRInput::BuildLegacyActions(TMap<FString, FInterac
 	{
 		// XR Keys are never blueprint-bindable
 		// 2D Axis keys are supported through paired keys
-		if (Key.IsBindableInBlueprints() || (Key.IsAnalog() && !Key.IsAxis1D()))
+		if (Key.IsDeprecated() || Key.IsBindableInBlueprints() || (Key.IsAnalog() && !Key.IsAxis1D()))
 		{
 			continue;
 		}
@@ -667,12 +667,6 @@ bool FOpenXRInputPlugin::FOpenXRInput::SuggestBindingForKey(TMap<FString, FInter
 			return false;
 		}
 
-		// The OpenXR spec says that .../input/system/click might not be available for application usage
-		if (Identifier == "system")
-		{
-			return false;
-		}
-
 		if (Identifier != "trigger" && Identifier != "squeeze")
 		{
 			Path += "/click";
diff --git a/Engine/Source/Runtime/InputCore/Classes/InputCoreTypes.h b/Engine/Source/Runtime/InputCore/Classes/InputCoreTypes.h
--- a/Engine/Source/Runtime/InputCore/Classes/InputCoreTypes.h
+++ b/Engine/Source/Runtime/InputCore/Classes/InputCoreTypes.h
@@ -503,6 +503,7 @@ struct INPUTCORE_API EKeys
 	static const FKey Android_Menu;
 
 	// HTC Vive Controller
+	UE_DEPRECATED(5.1, "This key has been deprecated and will be removed.")
 	static const FKey Vive_Left_System_Click;
 	static const FKey Vive_Left_Grip_Click;
 	static const FKey Vive_Left_Menu_Click;
@@ -517,6 +518,7 @@ struct INPUTCORE_API EKeys
 	static const FKey Vive_Left_Trackpad_Down;
 	static const FKey Vive_Left_Trackpad_Left;
 	static const FKey Vive_Left_Trackpad_Right;
+	UE_DEPRECATED(5.1, "This key has been deprecated and will be removed.")
 	static const FKey Vive_Right_System_Click;
 	static const FKey Vive_Right_Grip_Click;
 	static const FKey Vive_Right_Menu_Click;
@@ -600,6 +602,7 @@ struct INPUTCORE_API EKeys
 	static const FKey OculusTouch_Right_B_Click;
 	static const FKey OculusTouch_Right_A_Touch;
 	static const FKey OculusTouch_Right_B_Touch;
+	UE_DEPRECATED(5.1, "This key has been deprecated and will be removed.")
 	static const FKey OculusTouch_Right_System_Click;
 	static const FKey OculusTouch_Right_Grip_Click;
 	static const FKey OculusTouch_Right_Grip_Axis;
@@ -621,7 +624,9 @@ struct INPUTCORE_API EKeys
 	static const FKey ValveIndex_Left_B_Click;
 	static const FKey ValveIndex_Left_A_Touch;
 	static const FKey ValveIndex_Left_B_Touch;
+	UE_DEPRECATED(5.1, "This key has been deprecated and will be removed.")
 	static const FKey ValveIndex_Left_System_Click;
+	UE_DEPRECATED(5.1, "This key has been deprecated and will be removed.")
 	static const FKey ValveIndex_Left_System_Touch;
 	static const FKey ValveIndex_Left_Grip_Axis;
 	static const FKey ValveIndex_Left_Grip_Force;
@@ -650,7 +655,9 @@ struct INPUTCORE_API EKeys
 	static const FKey ValveIndex_Right_B_Click;
 	static const FKey ValveIndex_Right_A_Touch;
 	static const FKey ValveIndex_Right_B_Touch;
+	UE_DEPRECATED(5.1, "This key has been deprecated and will be removed.")
 	static const FKey ValveIndex_Right_System_Click;
+	UE_DEPRECATED(5.1, "This key has been deprecated and will be removed.")
 	static const FKey ValveIndex_Right_System_Touch;
 	static const FKey ValveIndex_Right_Grip_Axis;
 	static const FKey ValveIndex_Right_Grip_Force;
diff --git a/Engine/Source/Runtime/InputCore/Private/InputCoreTypes.cpp b/Engine/Source/Runtime/InputCore/Private/InputCoreTypes.cpp
--- a/Engine/Source/Runtime/InputCore/Private/InputCoreTypes.cpp
+++ b/Engine/Source/Runtime/InputCore/Private/InputCoreTypes.cpp
@@ -260,7 +260,9 @@ const FKey EKeys::Android_Volume_Down("Android_Volume_Down");
 const FKey EKeys::Android_Menu("Android_Menu");
 
 // HTC Vive Controller
+PRAGMA_DISABLE_DEPRECATION_WARNINGS
 const FKey EKeys::Vive_Left_System_Click("Vive_Left_System_Click");
+PRAGMA_ENABLE_DEPRECATION_WARNINGS
 const FKey EKeys::Vive_Left_Grip_Click("Vive_Left_Grip_Click");
 const FKey EKeys::Vive_Left_Menu_Click("Vive_Left_Menu_Click");
 const FKey EKeys::Vive_Left_Trigger_Click("Vive_Left_Trigger_Click");
@@ -274,7 +276,9 @@ const FKey EKeys::Vive_Left_Trackpad_Up("Vive_Left_Trackpad_Up");
 const FKey EKeys::Vive_Left_Trackpad_Down("Vive_Left_Trackpad_Down");
 const FKey EKeys::Vive_Left_Trackpad_Left("Vive_Left_Trackpad_Left");
 const FKey EKeys::Vive_Left_Trackpad_Right("Vive_Left_Trackpad_Right");
+PRAGMA_DISABLE_DEPRECATION_WARNINGS
 const FKey EKeys::Vive_Right_System_Click("Vive_Right_System_Click");
+PRAGMA_ENABLE_DEPRECATION_WARNINGS
 const FKey EKeys::Vive_Right_Grip_Click("Vive_Right_Grip_Click");
 const FKey EKeys::Vive_Right_Menu_Click("Vive_Right_Menu_Click");
 const FKey EKeys::Vive_Right_Trigger_Click("Vive_Right_Trigger_Click");
@@ -357,7 +361,9 @@ const FKey EKeys::OculusTouch_Right_A_Click("OculusTouch_Right_A_Click");
 const FKey EKeys::OculusTouch_Right_B_Click("OculusTouch_Right_B_Click");
 const FKey EKeys::OculusTouch_Right_A_Touch("OculusTouch_Right_A_Touch");
 const FKey EKeys::OculusTouch_Right_B_Touch("OculusTouch_Right_B_Touch");
+PRAGMA_DISABLE_DEPRECATION_WARNINGS
 const FKey EKeys::OculusTouch_Right_System_Click("OculusTouch_Right_System_Click");
+PRAGMA_ENABLE_DEPRECATION_WARNINGS
 const FKey EKeys::OculusTouch_Right_Grip_Click("OculusTouch_Right_Grip_Click");
 const FKey EKeys::OculusTouch_Right_Grip_Axis("OculusTouch_Right_Grip_Axis");
 const FKey EKeys::OculusTouch_Right_Trigger_Click("OculusTouch_Right_Trigger_Click");
@@ -378,8 +384,10 @@ const FKey EKeys::ValveIndex_Left_A_Click("ValveIndex_Left_A_Click");
 const FKey EKeys::ValveIndex_Left_B_Click("ValveIndex_Left_B_Click");
 const FKey EKeys::ValveIndex_Left_A_Touch("ValveIndex_Left_A_Touch");
 const FKey EKeys::ValveIndex_Left_B_Touch("ValveIndex_Left_B_Touch");
+PRAGMA_DISABLE_DEPRECATION_WARNINGS
 const FKey EKeys::ValveIndex_Left_System_Click("ValveIndex_Left_System_Click");
 const FKey EKeys::ValveIndex_Left_System_Touch("ValveIndex_Left_System_Touch");
+PRAGMA_ENABLE_DEPRECATION_WARNINGS
 const FKey EKeys::ValveIndex_Left_Grip_Axis("ValveIndex_Left_Grip_Axis");
 const FKey EKeys::ValveIndex_Left_Grip_Force("ValveIndex_Left_Grip_Force");
 const FKey EKeys::ValveIndex_Left_Trigger_Click("ValveIndex_Left_Trigger_Click");
@@ -407,8 +415,10 @@ const FKey EKeys::ValveIndex_Right_A_Click("ValveIndex_Right_A_Click");
 const FKey EKeys::ValveIndex_Right_B_Click("ValveIndex_Right_B_Click");
 const FKey EKeys::ValveIndex_Right_A_Touch("ValveIndex_Right_A_Touch");
 const FKey EKeys::ValveIndex_Right_B_Touch("ValveIndex_Right_B_Touch");
+PRAGMA_DISABLE_DEPRECATION_WARNINGS
 const FKey EKeys::ValveIndex_Right_System_Click("ValveIndex_Right_System_Click");
 const FKey EKeys::ValveIndex_Right_System_Touch("ValveIndex_Right_System_Touch");
+PRAGMA_ENABLE_DEPRECATION_WARNINGS
 const FKey EKeys::ValveIndex_Right_Grip_Axis("ValveIndex_Right_Grip_Axis");
 const FKey EKeys::ValveIndex_Right_Grip_Force("ValveIndex_Right_Grip_Force");
 const FKey EKeys::ValveIndex_Right_Trigger_Click("ValveIndex_Right_Trigger_Click");
@@ -789,7 +799,9 @@ void EKeys::Initialize()
 	// HTC Vive Controller
 	AddMenuCategoryDisplayInfo("Vive", LOCTEXT("ViveSubCategory", "HTC Vive"), TEXT("GraphEditor.PadEvent_16x"));
 
-	AddKey(FKeyDetails(EKeys::Vive_Left_System_Click, LOCTEXT("Vive_Left_System_Click", "Vive (L) System"), FKeyDetails::GamepadKey | FKeyDetails::NotBlueprintBindableKey, "Vive"));
+	PRAGMA_DISABLE_DEPRECATION_WARNINGS
+	AddKey(FKeyDetails(EKeys::Vive_Left_System_Click, LOCTEXT("Vive_Left_System_Click", "Vive (L) System"), FKeyDetails::Deprecated | FKeyDetails::GamepadKey | FKeyDetails::NotBlueprintBindableKey, "Vive"));
+	PRAGMA_ENABLE_DEPRECATION_WARNINGS
 	AddKey(FKeyDetails(EKeys::Vive_Left_Grip_Click, LOCTEXT("Vive_Left_Grip_Click", "Vive (L) Grip"), FKeyDetails::GamepadKey | FKeyDetails::NotBlueprintBindableKey, "Vive"));
 	AddKey(FKeyDetails(EKeys::Vive_Left_Menu_Click, LOCTEXT("Vive_Left_Menu_Click", "Vive (L) Menu"), FKeyDetails::GamepadKey | FKeyDetails::NotBlueprintBindableKey, "Vive"));
 	AddKey(FKeyDetails(EKeys::Vive_Left_Trigger_Click, LOCTEXT("Vive_Left_Trigger_Click", "Vive (L) Trigger"), FKeyDetails::GamepadKey | FKeyDetails::NotBlueprintBindableKey, "Vive"));
@@ -803,7 +815,9 @@ void EKeys::Initialize()
 	AddKey(FKeyDetails(EKeys::Vive_Left_Trackpad_Down, LOCTEXT("Vive_Left_Trackpad_Down", "Vive (L) Trackpad Down"), FKeyDetails::GamepadKey | FKeyDetails::NotBlueprintBindableKey, "Vive"));
 	AddKey(FKeyDetails(EKeys::Vive_Left_Trackpad_Left, LOCTEXT("Vive_Left_Trackpad_Left", "Vive (L) Trackpad Left"), FKeyDetails::GamepadKey | FKeyDetails::NotBlueprintBindableKey, "Vive"));
 	AddKey(FKeyDetails(EKeys::Vive_Left_Trackpad_Right, LOCTEXT("Vive_Left_Trackpad_Right", "Vive (L) Trackpad Right"), FKeyDetails::GamepadKey | FKeyDetails::NotBlueprintBindableKey, "Vive"));
-	AddKey(FKeyDetails(EKeys::Vive_Right_System_Click, LOCTEXT("Vive_Right_System_Click", "Vive (R) System"), FKeyDetails::GamepadKey | FKeyDetails::NotBlueprintBindableKey, "Vive"));
+	PRAGMA_DISABLE_DEPRECATION_WARNINGS
+	AddKey(FKeyDetails(EKeys::Vive_Right_System_Click, LOCTEXT("Vive_Right_System_Click", "Vive (R) System"), FKeyDetails::Deprecated | FKeyDetails::GamepadKey | FKeyDetails::NotBlueprintBindableKey, "Vive"));
+	PRAGMA_ENABLE_DEPRECATION_WARNINGS
 	AddKey(FKeyDetails(EKeys::Vive_Right_Grip_Click, LOCTEXT("Vive_Right_Grip_Click", "Vive (R) Grip"), FKeyDetails::GamepadKey | FKeyDetails::NotBlueprintBindableKey, "Vive"));
 	AddKey(FKeyDetails(EKeys::Vive_Right_Menu_Click, LOCTEXT("Vive_Right_Menu_Click", "Vive (R) Menu"), FKeyDetails::GamepadKey | FKeyDetails::NotBlueprintBindableKey, "Vive"));
 	AddKey(FKeyDetails(EKeys::Vive_Right_Trigger_Click, LOCTEXT("Vive_Right_Trigger_Click", "Vive (R) Trigger"), FKeyDetails::GamepadKey | FKeyDetails::NotBlueprintBindableKey, "Vive"));
@@ -890,7 +904,9 @@ void EKeys::Initialize()
 	AddKey(FKeyDetails(EKeys::OculusTouch_Right_B_Click, LOCTEXT("OculusTouch_Right_B_Click", "Oculus Touch (R) B Press"), FKeyDetails::GamepadKey | FKeyDetails::NotBlueprintBindableKey, "OculusTouch"));
 	AddKey(FKeyDetails(EKeys::OculusTouch_Right_A_Touch, LOCTEXT("OculusTouch_Right_A_Touch", "Oculus Touch (R) A Touch"), FKeyDetails::GamepadKey | FKeyDetails::NotBlueprintBindableKey, "OculusTouch"));
 	AddKey(FKeyDetails(EKeys::OculusTouch_Right_B_Touch, LOCTEXT("OculusTouch_Right_B_Touch", "Oculus Touch (R) B Touch"), FKeyDetails::GamepadKey | FKeyDetails::NotBlueprintBindableKey, "OculusTouch"));
-	AddKey(FKeyDetails(EKeys::OculusTouch_Right_System_Click, LOCTEXT("OculusTouch_Right_System_Click", "Oculus Touch (R) System"), FKeyDetails::GamepadKey | FKeyDetails::NotBlueprintBindableKey, "OculusTouch"));
+	PRAGMA_DISABLE_DEPRECATION_WARNINGS
+	AddKey(FKeyDetails(EKeys::OculusTouch_Right_System_Click, LOCTEXT("OculusTouch_Right_System_Click", "Oculus Touch (R) System"), FKeyDetails::Deprecated | FKeyDetails::GamepadKey | FKeyDetails::NotBlueprintBindableKey, "OculusTouch"));
+	PRAGMA_ENABLE_DEPRECATION_WARNINGS
 	AddKey(FKeyDetails(EKeys::OculusTouch_Right_Grip_Click, LOCTEXT("OculusTouch_Right_Grip_Click", "Oculus Touch (R) Grip"), FKeyDetails::GamepadKey | FKeyDetails::NotBlueprintBindableKey, "OculusTouch"));
 	AddKey(FKeyDetails(EKeys::OculusTouch_Right_Grip_Axis, LOCTEXT("OculusTouch_Right_Grip_Axis", "Oculus Touch (R) Grip Axis"), FKeyDetails::GamepadKey | FKeyDetails::Axis1D | FKeyDetails::NotBlueprintBindableKey, "OculusTouch"));
 	AddKey(FKeyDetails(EKeys::OculusTouch_Right_Trigger_Click, LOCTEXT("OculusTouch_Right_Trigger_Click", "Oculus Touch (R) Trigger"), FKeyDetails::GamepadKey | FKeyDetails::NotBlueprintBindableKey, "OculusTouch"));
@@ -913,8 +929,10 @@ void EKeys::Initialize()
 	AddKey(FKeyDetails(EKeys::ValveIndex_Left_B_Click, LOCTEXT("ValveIndex_Left_B_Click", "Valve Index (L) B Press"), FKeyDetails::GamepadKey | FKeyDetails::NotBlueprintBindableKey, "ValveIndex"));
 	AddKey(FKeyDetails(EKeys::ValveIndex_Left_A_Touch, LOCTEXT("ValveIndex_Left_A_Touch", "Valve Index (L) A Touch"), FKeyDetails::GamepadKey | FKeyDetails::NotBlueprintBindableKey, "ValveIndex"));
 	AddKey(FKeyDetails(EKeys::ValveIndex_Left_B_Touch, LOCTEXT("ValveIndex_Left_B_Touch", "Valve Index (L) B Touch"), FKeyDetails::GamepadKey | FKeyDetails::NotBlueprintBindableKey, "ValveIndex"));
-	AddKey(FKeyDetails(EKeys::ValveIndex_Left_System_Click, LOCTEXT("ValveIndex_Left_System_Click", "Valve Index (L) System"), FKeyDetails::GamepadKey | FKeyDetails::NotBlueprintBindableKey, "ValveIndex"));
-	AddKey(FKeyDetails(EKeys::ValveIndex_Left_System_Touch, LOCTEXT("ValveIndex_Left_System_Touch", "Valve Index (L) System Touch"), FKeyDetails::GamepadKey | FKeyDetails::NotBlueprintBindableKey, "ValveIndex"));
+	PRAGMA_DISABLE_DEPRECATION_WARNINGS
+	AddKey(FKeyDetails(EKeys::ValveIndex_Left_System_Click, LOCTEXT("ValveIndex_Left_System_Click", "Valve Index (L) System"), FKeyDetails::Deprecated | FKeyDetails::GamepadKey | FKeyDetails::NotBlueprintBindableKey, "ValveIndex"));
+	AddKey(FKeyDetails(EKeys::ValveIndex_Left_System_Touch, LOCTEXT("ValveIndex_Left_System_Touch", "Valve Index (L) System Touch"), FKeyDetails::Deprecated | FKeyDetails::GamepadKey | FKeyDetails::NotBlueprintBindableKey, "ValveIndex"));
+	PRAGMA_ENABLE_DEPRECATION_WARNINGS
 	AddKey(FKeyDetails(EKeys::ValveIndex_Left_Grip_Axis, LOCTEXT("ValveIndex_Left_Grip_Axis", "Valve Index (L) Grip Axis"), FKeyDetails::GamepadKey | FKeyDetails::Axis1D | FKeyDetails::NotBlueprintBindableKey, "ValveIndex"));
 	AddKey(FKeyDetails(EKeys::ValveIndex_Left_Grip_Force, LOCTEXT("ValveIndex_Left_Grip_Force", "Valve Index (L) Grip Force"), FKeyDetails::GamepadKey | FKeyDetails::Axis1D | FKeyDetails::NotBlueprintBindableKey, "ValveIndex"));
 	AddKey(FKeyDetails(EKeys::ValveIndex_Left_Trigger_Click, LOCTEXT("ValveIndex_Left_Trigger_Click", "Valve Index (L) Trigger"), FKeyDetails::GamepadKey | FKeyDetails::NotBlueprintBindableKey, "ValveIndex"));
@@ -942,8 +960,10 @@ void EKeys::Initialize()
 	AddKey(FKeyDetails(EKeys::ValveIndex_Right_B_Click, LOCTEXT("ValveIndex_Right_B_Click", "Valve Index (R) B Press"), FKeyDetails::GamepadKey | FKeyDetails::NotBlueprintBindableKey, "ValveIndex"));
 	AddKey(FKeyDetails(EKeys::ValveIndex_Right_A_Touch, LOCTEXT("ValveIndex_Right_A_Touch", "Valve Index (R) A Touch"), FKeyDetails::GamepadKey | FKeyDetails::NotBlueprintBindableKey, "ValveIndex"));
 	AddKey(FKeyDetails(EKeys::ValveIndex_Right_B_Touch, LOCTEXT("ValveIndex_Right_B_Touch", "Valve Index (R) B Touch"), FKeyDetails::GamepadKey | FKeyDetails::NotBlueprintBindableKey, "ValveIndex"));
-	AddKey(FKeyDetails(EKeys::ValveIndex_Right_System_Click, LOCTEXT("ValveIndex_Right_System_Click", "Valve Index (R) System"), FKeyDetails::GamepadKey | FKeyDetails::NotBlueprintBindableKey, "ValveIndex"));
-	AddKey(FKeyDetails(EKeys::ValveIndex_Right_System_Touch, LOCTEXT("ValveIndex_Right_System_Touch", "Valve Index (R) System Touch"), FKeyDetails::GamepadKey | FKeyDetails::NotBlueprintBindableKey, "ValveIndex"));
+	PRAGMA_DISABLE_DEPRECATION_WARNINGS
+	AddKey(FKeyDetails(EKeys::ValveIndex_Right_System_Click, LOCTEXT("ValveIndex_Right_System_Click", "Valve Index (R) System"), FKeyDetails::Deprecated | FKeyDetails::GamepadKey | FKeyDetails::NotBlueprintBindableKey, "ValveIndex"));
+	AddKey(FKeyDetails(EKeys::ValveIndex_Right_System_Touch, LOCTEXT("ValveIndex_Right_System_Touch", "Valve Index (R) System Touch"), FKeyDetails::Deprecated | FKeyDetails::GamepadKey | FKeyDetails::NotBlueprintBindableKey, "ValveIndex"));
+	PRAGMA_ENABLE_DEPRECATION_WARNINGS
 	AddKey(FKeyDetails(EKeys::ValveIndex_Right_Grip_Axis, LOCTEXT("ValveIndex_Right_Grip_Axis", "Valve Index (R) Grip Axis"), FKeyDetails::GamepadKey | FKeyDetails::Axis1D | FKeyDetails::NotBlueprintBindableKey, "ValveIndex"));
 	AddKey(FKeyDetails(EKeys::ValveIndex_Right_Grip_Force, LOCTEXT("ValveIndex_Right_Grip_Force", "Valve Index (R) Grip Force"), FKeyDetails::GamepadKey | FKeyDetails::Axis1D | FKeyDetails::NotBlueprintBindableKey, "ValveIndex"));
 	AddKey(FKeyDetails(EKeys::ValveIndex_Right_Trigger_Click, LOCTEXT("ValveIndex_Right_Trigger_Click", "Valve Index (R) Trigger"), FKeyDetails::GamepadKey | FKeyDetails::NotBlueprintBindableKey, "ValveIndex"));
