Mark the FKey::KeyName as EditAnywhere so that its value is correctly propegated if it is changed on a subobject. Without this, the KeyName property will be filtered out when determining if a property has changed on a subobject, leading to an unusable UPROPERTY. This has no effect visually in the editor since the FKey has a custom details row.

diff --git a/Engine/Source/Runtime/InputCore/Classes/InputCoreTypes.h b/Engine/Source/Runtime/InputCore/Classes/InputCoreTypes.h
--- a/Engine/Source/Runtime/InputCore/Classes/InputCoreTypes.h
+++ b/Engine/Source/Runtime/InputCore/Classes/InputCoreTypes.h
@@ -113,7 +113,7 @@ struct INPUTCORE_API FKey
 
 private:
 
-	UPROPERTY()
+	UPROPERTY(EditAnywhere, Category="Input")
 	FName KeyName;
 
 	mutable class TSharedPtr<struct FKeyDetails> KeyDetails;
