[Engine]

diff --git a/Engine/Source/Runtime/InputCore/Classes/InputCoreTypes.h b/Engine/Source/Runtime/InputCore/Classes/InputCoreTypes.h
--- a/Engine/Source/Runtime/InputCore/Classes/InputCoreTypes.h
+++ b/Engine/Source/Runtime/InputCore/Classes/InputCoreTypes.h
@@ -138,7 +138,7 @@ struct TStructOpsTypeTraits<FKey> : public TStructOpsTypeTraitsBase2<FKey>
 
 DECLARE_DELEGATE_RetVal_OneParam(FText, FGetKeyDisplayNameSignature, const FKey);
 
-struct INPUTCORE_API FKeyDetails
+struct FKeyDetails
 {
 	enum EKeyFlags
 	{
@@ -168,8 +168,8 @@ struct INPUTCORE_API FKeyDetails
 		NoFlags                  = 0,
 	};
 
-	FKeyDetails(const FKey InKey, const TAttribute<FText>& InLongDisplayName, const uint32 InKeyFlags = 0, const FName InMenuCategory = NAME_None, const TAttribute<FText>& InShortDisplayName = TAttribute<FText>() );
-	FKeyDetails(const FKey InKey, const TAttribute<FText>& InLongDisplayName, const TAttribute<FText>& InShortDisplayName, const uint32 InKeyFlags = 0, const FName InMenuCategory = NAME_None);
+	INPUTCORE_API FKeyDetails(const FKey InKey, const TAttribute<FText>& InLongDisplayName, const uint32 InKeyFlags = 0, const FName InMenuCategory = NAME_None, const TAttribute<FText>& InShortDisplayName = TAttribute<FText>() );
+	INPUTCORE_API FKeyDetails(const FKey InKey, const TAttribute<FText>& InLongDisplayName, const TAttribute<FText>& InShortDisplayName, const uint32 InKeyFlags = 0, const FName InMenuCategory = NAME_None);
 
 	FORCEINLINE bool IsModifierKey() const { return bIsModifierKey != 0; }
 	FORCEINLINE bool IsGamepadKey() const { return bIsGamepadKey != 0; }
@@ -191,7 +191,7 @@ struct INPUTCORE_API FKeyDetails
 	FORCEINLINE bool IsGesture() const { return bIsGesture != 0; }
 	FORCEINLINE bool IsDeprecated() const { return bIsDeprecated != 0; }
 	FORCEINLINE FName GetMenuCategory() const { return MenuCategory; }
-	FText GetDisplayName(const bool bLongDisplayName = true) const;
+	INPUTCORE_API FText GetDisplayName(const bool bLongDisplayName = true) const;
 	FORCEINLINE const FKey& GetKey() const { return Key; }
 
 	// Key pairing
@@ -201,7 +201,7 @@ struct INPUTCORE_API FKeyDetails
 private:
 	friend struct EKeys;
 
-	void CommonInit(const uint32 InKeyFlags);
+	INPUTCORE_API void CommonInit(const uint32 InKeyFlags);
 
 	enum class EInputAxisType : uint8
 	{
@@ -271,441 +271,441 @@ namespace EConsoleForGamepadLabels
 	};
 }
 
-struct INPUTCORE_API EKeys
+struct EKeys
 {
-	static const FKey AnyKey;
-
-	static const FKey MouseX;
-	static const FKey MouseY;
-	static const FKey Mouse2D;
-	static const FKey MouseScrollUp;
-	static const FKey MouseScrollDown;
-	static const FKey MouseWheelAxis;
-
-	static const FKey LeftMouseButton;
-	static const FKey RightMouseButton;
-	static const FKey MiddleMouseButton;
-	static const FKey ThumbMouseButton;
-	static const FKey ThumbMouseButton2;
-
-	static const FKey BackSpace;
-	static const FKey Tab;
-	static const FKey Enter;
-	static const FKey Pause;
-
-	static const FKey CapsLock;
-	static const FKey Escape;
-	static const FKey SpaceBar;
-	static const FKey PageUp;
-	static const FKey PageDown;
-	static const FKey End;
-	static const FKey Home;
-
-	static const FKey Left;
-	static const FKey Up;
-	static const FKey Right;
-	static const FKey Down;
-
-	static const FKey Insert;
-	static const FKey Delete;
-
-	static const FKey Zero;
-	static const FKey One;
-	static const FKey Two;
-	static const FKey Three;
-	static const FKey Four;
-	static const FKey Five;
-	static const FKey Six;
-	static const FKey Seven;
-	static const FKey Eight;
-	static const FKey Nine;
-
-	static const FKey A;
-	static const FKey B;
-	static const FKey C;
-	static const FKey D;
-	static const FKey E;
-	static const FKey F;
-	static const FKey G;
-	static const FKey H;
-	static const FKey I;
-	static const FKey J;
-	static const FKey K;
-	static const FKey L;
-	static const FKey M;
-	static const FKey N;
-	static const FKey O;
-	static const FKey P;
-	static const FKey Q;
-	static const FKey R;
-	static const FKey S;
-	static const FKey T;
-	static const FKey U;
-	static const FKey V;
-	static const FKey W;
-	static const FKey X;
-	static const FKey Y;
-	static const FKey Z;
-
-	static const FKey NumPadZero;
-	static const FKey NumPadOne;
-	static const FKey NumPadTwo;
-	static const FKey NumPadThree;
-	static const FKey NumPadFour;
-	static const FKey NumPadFive;
-	static const FKey NumPadSix;
-	static const FKey NumPadSeven;
-	static const FKey NumPadEight;
-	static const FKey NumPadNine;
-
-	static const FKey Multiply;
-	static const FKey Add;
-	static const FKey Subtract;
-	static const FKey Decimal;
-	static const FKey Divide;
-
-	static const FKey F1;
-	static const FKey F2;
-	static const FKey F3;
-	static const FKey F4;
-	static const FKey F5;
-	static const FKey F6;
-	static const FKey F7;
-	static const FKey F8;
-	static const FKey F9;
-	static const FKey F10;
-	static const FKey F11;
-	static const FKey F12;
-
-	static const FKey NumLock;
-
-	static const FKey ScrollLock;
-
-	static const FKey LeftShift;
-	static const FKey RightShift;
-	static const FKey LeftControl;
-	static const FKey RightControl;
-	static const FKey LeftAlt;
-	static const FKey RightAlt;
-	static const FKey LeftCommand;
-	static const FKey RightCommand;
-
-	static const FKey Semicolon;
-	static const FKey Equals;
-	static const FKey Comma;
-	static const FKey Underscore;
-	static const FKey Hyphen;
-	static const FKey Period;
-	static const FKey Slash;
-	static const FKey Tilde;
-	static const FKey LeftBracket;
-	static const FKey Backslash;
-	static const FKey RightBracket;
-	static const FKey Apostrophe;
-
-	static const FKey Ampersand;
-	static const FKey Asterix;
-	static const FKey Caret;
-	static const FKey Colon;
-	static const FKey Dollar;
-	static const FKey Exclamation;
-	static const FKey LeftParantheses;
-	static const FKey RightParantheses;
-	static const FKey Quote;
-
-	static const FKey A_AccentGrave;
-	static const FKey E_AccentGrave;
-	static const FKey E_AccentAigu;
-	static const FKey C_Cedille;
-	static const FKey Section;
+	static INPUTCORE_API const FKey AnyKey;
+
+	static INPUTCORE_API const FKey MouseX;
+	static INPUTCORE_API const FKey MouseY;
+	static INPUTCORE_API const FKey Mouse2D;
+	static INPUTCORE_API const FKey MouseScrollUp;
+	static INPUTCORE_API const FKey MouseScrollDown;
+	static INPUTCORE_API const FKey MouseWheelAxis;
+
+	static INPUTCORE_API const FKey LeftMouseButton;
+	static INPUTCORE_API const FKey RightMouseButton;
+	static INPUTCORE_API const FKey MiddleMouseButton;
+	static INPUTCORE_API const FKey ThumbMouseButton;
+	static INPUTCORE_API const FKey ThumbMouseButton2;
+
+	static INPUTCORE_API const FKey BackSpace;
+	static INPUTCORE_API const FKey Tab;
+	static INPUTCORE_API const FKey Enter;
+	static INPUTCORE_API const FKey Pause;
+
+	static INPUTCORE_API const FKey CapsLock;
+	static INPUTCORE_API const FKey Escape;
+	static INPUTCORE_API const FKey SpaceBar;
+	static INPUTCORE_API const FKey PageUp;
+	static INPUTCORE_API const FKey PageDown;
+	static INPUTCORE_API const FKey End;
+	static INPUTCORE_API const FKey Home;
+
+	static INPUTCORE_API const FKey Left;
+	static INPUTCORE_API const FKey Up;
+	static INPUTCORE_API const FKey Right;
+	static INPUTCORE_API const FKey Down;
+
+	static INPUTCORE_API const FKey Insert;
+	static INPUTCORE_API const FKey Delete;
+
+	static INPUTCORE_API const FKey Zero;
+	static INPUTCORE_API const FKey One;
+	static INPUTCORE_API const FKey Two;
+	static INPUTCORE_API const FKey Three;
+	static INPUTCORE_API const FKey Four;
+	static INPUTCORE_API const FKey Five;
+	static INPUTCORE_API const FKey Six;
+	static INPUTCORE_API const FKey Seven;
+	static INPUTCORE_API const FKey Eight;
+	static INPUTCORE_API const FKey Nine;
+
+	static INPUTCORE_API const FKey A;
+	static INPUTCORE_API const FKey B;
+	static INPUTCORE_API const FKey C;
+	static INPUTCORE_API const FKey D;
+	static INPUTCORE_API const FKey E;
+	static INPUTCORE_API const FKey F;
+	static INPUTCORE_API const FKey G;
+	static INPUTCORE_API const FKey H;
+	static INPUTCORE_API const FKey I;
+	static INPUTCORE_API const FKey J;
+	static INPUTCORE_API const FKey K;
+	static INPUTCORE_API const FKey L;
+	static INPUTCORE_API const FKey M;
+	static INPUTCORE_API const FKey N;
+	static INPUTCORE_API const FKey O;
+	static INPUTCORE_API const FKey P;
+	static INPUTCORE_API const FKey Q;
+	static INPUTCORE_API const FKey R;
+	static INPUTCORE_API const FKey S;
+	static INPUTCORE_API const FKey T;
+	static INPUTCORE_API const FKey U;
+	static INPUTCORE_API const FKey V;
+	static INPUTCORE_API const FKey W;
+	static INPUTCORE_API const FKey X;
+	static INPUTCORE_API const FKey Y;
+	static INPUTCORE_API const FKey Z;
+
+	static INPUTCORE_API const FKey NumPadZero;
+	static INPUTCORE_API const FKey NumPadOne;
+	static INPUTCORE_API const FKey NumPadTwo;
+	static INPUTCORE_API const FKey NumPadThree;
+	static INPUTCORE_API const FKey NumPadFour;
+	static INPUTCORE_API const FKey NumPadFive;
+	static INPUTCORE_API const FKey NumPadSix;
+	static INPUTCORE_API const FKey NumPadSeven;
+	static INPUTCORE_API const FKey NumPadEight;
+	static INPUTCORE_API const FKey NumPadNine;
+
+	static INPUTCORE_API const FKey Multiply;
+	static INPUTCORE_API const FKey Add;
+	static INPUTCORE_API const FKey Subtract;
+	static INPUTCORE_API const FKey Decimal;
+	static INPUTCORE_API const FKey Divide;
+
+	static INPUTCORE_API const FKey F1;
+	static INPUTCORE_API const FKey F2;
+	static INPUTCORE_API const FKey F3;
+	static INPUTCORE_API const FKey F4;
+	static INPUTCORE_API const FKey F5;
+	static INPUTCORE_API const FKey F6;
+	static INPUTCORE_API const FKey F7;
+	static INPUTCORE_API const FKey F8;
+	static INPUTCORE_API const FKey F9;
+	static INPUTCORE_API const FKey F10;
+	static INPUTCORE_API const FKey F11;
+	static INPUTCORE_API const FKey F12;
+
+	static INPUTCORE_API const FKey NumLock;
+
+	static INPUTCORE_API const FKey ScrollLock;
+
+	static INPUTCORE_API const FKey LeftShift;
+	static INPUTCORE_API const FKey RightShift;
+	static INPUTCORE_API const FKey LeftControl;
+	static INPUTCORE_API const FKey RightControl;
+	static INPUTCORE_API const FKey LeftAlt;
+	static INPUTCORE_API const FKey RightAlt;
+	static INPUTCORE_API const FKey LeftCommand;
+	static INPUTCORE_API const FKey RightCommand;
+
+	static INPUTCORE_API const FKey Semicolon;
+	static INPUTCORE_API const FKey Equals;
+	static INPUTCORE_API const FKey Comma;
+	static INPUTCORE_API const FKey Underscore;
+	static INPUTCORE_API const FKey Hyphen;
+	static INPUTCORE_API const FKey Period;
+	static INPUTCORE_API const FKey Slash;
+	static INPUTCORE_API const FKey Tilde;
+	static INPUTCORE_API const FKey LeftBracket;
+	static INPUTCORE_API const FKey Backslash;
+	static INPUTCORE_API const FKey RightBracket;
+	static INPUTCORE_API const FKey Apostrophe;
+
+	static INPUTCORE_API const FKey Ampersand;
+	static INPUTCORE_API const FKey Asterix;
+	static INPUTCORE_API const FKey Caret;
+	static INPUTCORE_API const FKey Colon;
+	static INPUTCORE_API const FKey Dollar;
+	static INPUTCORE_API const FKey Exclamation;
+	static INPUTCORE_API const FKey LeftParantheses;
+	static INPUTCORE_API const FKey RightParantheses;
+	static INPUTCORE_API const FKey Quote;
+
+	static INPUTCORE_API const FKey A_AccentGrave;
+	static INPUTCORE_API const FKey E_AccentGrave;
+	static INPUTCORE_API const FKey E_AccentAigu;
+	static INPUTCORE_API const FKey C_Cedille;
+	static INPUTCORE_API const FKey Section;
 
 	// Platform Keys
 	// These keys platform specific versions of keys that go by different names.
 	// The delete key is a good example, on Windows Delete is the virtual key for Delete.
 	// On Macs, the Delete key is the virtual key for BackSpace.
-	static const FKey Platform_Delete;
+	static INPUTCORE_API const FKey Platform_Delete;
 
 	// Gamepad Keys
-	static const FKey Gamepad_Left2D;
-	static const FKey Gamepad_LeftX;
-	static const FKey Gamepad_LeftY;
-	static const FKey Gamepad_Right2D;
-	static const FKey Gamepad_RightX;
-	static const FKey Gamepad_RightY;
-	static const FKey Gamepad_LeftTriggerAxis;
-	static const FKey Gamepad_RightTriggerAxis;
-
-	static const FKey Gamepad_LeftThumbstick;
-	static const FKey Gamepad_RightThumbstick;
-	static const FKey Gamepad_Special_Left;
-	static const FKey Gamepad_Special_Left_X;
-	static const FKey Gamepad_Special_Left_Y;
-	static const FKey Gamepad_Special_Right;
-	static const FKey Gamepad_FaceButton_Bottom;
-	static const FKey Gamepad_FaceButton_Right;
-	static const FKey Gamepad_FaceButton_Left;
-	static const FKey Gamepad_FaceButton_Top;
-	static const FKey Gamepad_LeftShoulder;
-	static const FKey Gamepad_RightShoulder;
-	static const FKey Gamepad_LeftTrigger;
-	static const FKey Gamepad_RightTrigger;
-	static const FKey Gamepad_DPad_Up;
-	static const FKey Gamepad_DPad_Down;
-	static const FKey Gamepad_DPad_Right;
-	static const FKey Gamepad_DPad_Left;
+	static INPUTCORE_API const FKey Gamepad_Left2D;
+	static INPUTCORE_API const FKey Gamepad_LeftX;
+	static INPUTCORE_API const FKey Gamepad_LeftY;
+	static INPUTCORE_API const FKey Gamepad_Right2D;
+	static INPUTCORE_API const FKey Gamepad_RightX;
+	static INPUTCORE_API const FKey Gamepad_RightY;
+	static INPUTCORE_API const FKey Gamepad_LeftTriggerAxis;
+	static INPUTCORE_API const FKey Gamepad_RightTriggerAxis;
+
+	static INPUTCORE_API const FKey Gamepad_LeftThumbstick;
+	static INPUTCORE_API const FKey Gamepad_RightThumbstick;
+	static INPUTCORE_API const FKey Gamepad_Special_Left;
+	static INPUTCORE_API const FKey Gamepad_Special_Left_X;
+	static INPUTCORE_API const FKey Gamepad_Special_Left_Y;
+	static INPUTCORE_API const FKey Gamepad_Special_Right;
+	static INPUTCORE_API const FKey Gamepad_FaceButton_Bottom;
+	static INPUTCORE_API const FKey Gamepad_FaceButton_Right;
+	static INPUTCORE_API const FKey Gamepad_FaceButton_Left;
+	static INPUTCORE_API const FKey Gamepad_FaceButton_Top;
+	static INPUTCORE_API const FKey Gamepad_LeftShoulder;
+	static INPUTCORE_API const FKey Gamepad_RightShoulder;
+	static INPUTCORE_API const FKey Gamepad_LeftTrigger;
+	static INPUTCORE_API const FKey Gamepad_RightTrigger;
+	static INPUTCORE_API const FKey Gamepad_DPad_Up;
+	static INPUTCORE_API const FKey Gamepad_DPad_Down;
+	static INPUTCORE_API const FKey Gamepad_DPad_Right;
+	static INPUTCORE_API const FKey Gamepad_DPad_Left;
 
 	// Virtual key codes used for input axis button press/release emulation
-	static const FKey Gamepad_LeftStick_Up;
-	static const FKey Gamepad_LeftStick_Down;
-	static const FKey Gamepad_LeftStick_Right;
-	static const FKey Gamepad_LeftStick_Left;
+	static INPUTCORE_API const FKey Gamepad_LeftStick_Up;
+	static INPUTCORE_API const FKey Gamepad_LeftStick_Down;
+	static INPUTCORE_API const FKey Gamepad_LeftStick_Right;
+	static INPUTCORE_API const FKey Gamepad_LeftStick_Left;
 
-	static const FKey Gamepad_RightStick_Up;
-	static const FKey Gamepad_RightStick_Down;
-	static const FKey Gamepad_RightStick_Right;
-	static const FKey Gamepad_RightStick_Left;
+	static INPUTCORE_API const FKey Gamepad_RightStick_Up;
+	static INPUTCORE_API const FKey Gamepad_RightStick_Down;
+	static INPUTCORE_API const FKey Gamepad_RightStick_Right;
+	static INPUTCORE_API const FKey Gamepad_RightStick_Left;
 
 	// static const FKey Vector axes (FVector; not float)
-	static const FKey Tilt;
-	static const FKey RotationRate;
-	static const FKey Gravity;
-	static const FKey Acceleration;
+	static INPUTCORE_API const FKey Tilt;
+	static INPUTCORE_API const FKey RotationRate;
+	static INPUTCORE_API const FKey Gravity;
+	static INPUTCORE_API const FKey Acceleration;
 
 	// Gestures
-	static const FKey Gesture_Pinch;
-	static const FKey Gesture_Flick;
-	static const FKey Gesture_Rotate;
+	static INPUTCORE_API const FKey Gesture_Pinch;
+	static INPUTCORE_API const FKey Gesture_Flick;
+	static INPUTCORE_API const FKey Gesture_Rotate;
 
 	// PS4-specific
 	UE_DEPRECATED(5.0, "This key has deprecated and will be removed. Use GamePad_Special_Left/Right instead.")
-	static const FKey PS4_Special;
+	static INPUTCORE_API const FKey PS4_Special;
 
 	// Steam Controller Specific
-	static const FKey Steam_Touch_0;
-	static const FKey Steam_Touch_1;
-	static const FKey Steam_Touch_2;
-	static const FKey Steam_Touch_3;
-	static const FKey Steam_Back_Left;
-	static const FKey Steam_Back_Right;
+	static INPUTCORE_API const FKey Steam_Touch_0;
+	static INPUTCORE_API const FKey Steam_Touch_1;
+	static INPUTCORE_API const FKey Steam_Touch_2;
+	static INPUTCORE_API const FKey Steam_Touch_3;
+	static INPUTCORE_API const FKey Steam_Back_Left;
+	static INPUTCORE_API const FKey Steam_Back_Right;
 
 	// Xbox One global speech commands
-	static const FKey Global_Menu;
-	static const FKey Global_View;
-	static const FKey Global_Pause;
-	static const FKey Global_Play;
-	static const FKey Global_Back;
+	static INPUTCORE_API const FKey Global_Menu;
+	static INPUTCORE_API const FKey Global_View;
+	static INPUTCORE_API const FKey Global_Pause;
+	static INPUTCORE_API const FKey Global_Play;
+	static INPUTCORE_API const FKey Global_Back;
 
 	// Android-specific
-	static const FKey Android_Back;
-	static const FKey Android_Volume_Up;
-	static const FKey Android_Volume_Down;
-	static const FKey Android_Menu;
+	static INPUTCORE_API const FKey Android_Back;
+	static INPUTCORE_API const FKey Android_Volume_Up;
+	static INPUTCORE_API const FKey Android_Volume_Down;
+	static INPUTCORE_API const FKey Android_Menu;
 
 	// HTC Vive Controller
 	UE_DEPRECATED(5.1, "This key has been deprecated and will be removed.")
-	static const FKey Vive_Left_System_Click;
-	static const FKey Vive_Left_Grip_Click;
-	static const FKey Vive_Left_Menu_Click;
-	static const FKey Vive_Left_Trigger_Click;
-	static const FKey Vive_Left_Trigger_Axis;
-	static const FKey Vive_Left_Trackpad_2D;
-	static const FKey Vive_Left_Trackpad_X;
-	static const FKey Vive_Left_Trackpad_Y;
-	static const FKey Vive_Left_Trackpad_Click;
-	static const FKey Vive_Left_Trackpad_Touch;
-	static const FKey Vive_Left_Trackpad_Up;
-	static const FKey Vive_Left_Trackpad_Down;
-	static const FKey Vive_Left_Trackpad_Left;
-	static const FKey Vive_Left_Trackpad_Right;
+	static INPUTCORE_API const FKey Vive_Left_System_Click;
+	static INPUTCORE_API const FKey Vive_Left_Grip_Click;
+	static INPUTCORE_API const FKey Vive_Left_Menu_Click;
+	static INPUTCORE_API const FKey Vive_Left_Trigger_Click;
+	static INPUTCORE_API const FKey Vive_Left_Trigger_Axis;
+	static INPUTCORE_API const FKey Vive_Left_Trackpad_2D;
+	static INPUTCORE_API const FKey Vive_Left_Trackpad_X;
+	static INPUTCORE_API const FKey Vive_Left_Trackpad_Y;
+	static INPUTCORE_API const FKey Vive_Left_Trackpad_Click;
+	static INPUTCORE_API const FKey Vive_Left_Trackpad_Touch;
+	static INPUTCORE_API const FKey Vive_Left_Trackpad_Up;
+	static INPUTCORE_API const FKey Vive_Left_Trackpad_Down;
+	static INPUTCORE_API const FKey Vive_Left_Trackpad_Left;
+	static INPUTCORE_API const FKey Vive_Left_Trackpad_Right;
 	UE_DEPRECATED(5.1, "This key has been deprecated and will be removed.")
-	static const FKey Vive_Right_System_Click;
-	static const FKey Vive_Right_Grip_Click;
-	static const FKey Vive_Right_Menu_Click;
-	static const FKey Vive_Right_Trigger_Click;
-	static const FKey Vive_Right_Trigger_Axis;
-	static const FKey Vive_Right_Trackpad_2D;
-	static const FKey Vive_Right_Trackpad_X;
-	static const FKey Vive_Right_Trackpad_Y;
-	static const FKey Vive_Right_Trackpad_Click;
-	static const FKey Vive_Right_Trackpad_Touch;
-	static const FKey Vive_Right_Trackpad_Up;
-	static const FKey Vive_Right_Trackpad_Down;
-	static const FKey Vive_Right_Trackpad_Left;
-	static const FKey Vive_Right_Trackpad_Right;
+	static INPUTCORE_API const FKey Vive_Right_System_Click;
+	static INPUTCORE_API const FKey Vive_Right_Grip_Click;
+	static INPUTCORE_API const FKey Vive_Right_Menu_Click;
+	static INPUTCORE_API const FKey Vive_Right_Trigger_Click;
+	static INPUTCORE_API const FKey Vive_Right_Trigger_Axis;
+	static INPUTCORE_API const FKey Vive_Right_Trackpad_2D;
+	static INPUTCORE_API const FKey Vive_Right_Trackpad_X;
+	static INPUTCORE_API const FKey Vive_Right_Trackpad_Y;
+	static INPUTCORE_API const FKey Vive_Right_Trackpad_Click;
+	static INPUTCORE_API const FKey Vive_Right_Trackpad_Touch;
+	static INPUTCORE_API const FKey Vive_Right_Trackpad_Up;
+	static INPUTCORE_API const FKey Vive_Right_Trackpad_Down;
+	static INPUTCORE_API const FKey Vive_Right_Trackpad_Left;
+	static INPUTCORE_API const FKey Vive_Right_Trackpad_Right;
 
 	// Microsoft Mixed Reality Motion Controller
-	static const FKey MixedReality_Left_Menu_Click;
-	static const FKey MixedReality_Left_Grip_Click;
-	static const FKey MixedReality_Left_Trigger_Click;
-	static const FKey MixedReality_Left_Trigger_Axis;
-	static const FKey MixedReality_Left_Thumbstick_2D;
-	static const FKey MixedReality_Left_Thumbstick_X;
-	static const FKey MixedReality_Left_Thumbstick_Y;
-	static const FKey MixedReality_Left_Thumbstick_Click;
-	static const FKey MixedReality_Left_Thumbstick_Up;
-	static const FKey MixedReality_Left_Thumbstick_Down;
-	static const FKey MixedReality_Left_Thumbstick_Left;
-	static const FKey MixedReality_Left_Thumbstick_Right;
-	static const FKey MixedReality_Left_Trackpad_2D;
-	static const FKey MixedReality_Left_Trackpad_X;
-	static const FKey MixedReality_Left_Trackpad_Y;
-	static const FKey MixedReality_Left_Trackpad_Click;
-	static const FKey MixedReality_Left_Trackpad_Touch;
-	static const FKey MixedReality_Left_Trackpad_Up;
-	static const FKey MixedReality_Left_Trackpad_Down;
-	static const FKey MixedReality_Left_Trackpad_Left;
-	static const FKey MixedReality_Left_Trackpad_Right;
-	static const FKey MixedReality_Right_Menu_Click;
-	static const FKey MixedReality_Right_Grip_Click;
-	static const FKey MixedReality_Right_Trigger_Click;
-	static const FKey MixedReality_Right_Trigger_Axis;
-	static const FKey MixedReality_Right_Thumbstick_2D;
-	static const FKey MixedReality_Right_Thumbstick_X;
-	static const FKey MixedReality_Right_Thumbstick_Y;
-	static const FKey MixedReality_Right_Thumbstick_Click;
-	static const FKey MixedReality_Right_Thumbstick_Up;
-	static const FKey MixedReality_Right_Thumbstick_Down;
-	static const FKey MixedReality_Right_Thumbstick_Left;
-	static const FKey MixedReality_Right_Thumbstick_Right;
-	static const FKey MixedReality_Right_Trackpad_2D;
-	static const FKey MixedReality_Right_Trackpad_X;
-	static const FKey MixedReality_Right_Trackpad_Y;
-	static const FKey MixedReality_Right_Trackpad_Click;
-	static const FKey MixedReality_Right_Trackpad_Touch;
-	static const FKey MixedReality_Right_Trackpad_Up;
-	static const FKey MixedReality_Right_Trackpad_Down;
-	static const FKey MixedReality_Right_Trackpad_Left;
-	static const FKey MixedReality_Right_Trackpad_Right;
+	static INPUTCORE_API const FKey MixedReality_Left_Menu_Click;
+	static INPUTCORE_API const FKey MixedReality_Left_Grip_Click;
+	static INPUTCORE_API const FKey MixedReality_Left_Trigger_Click;
+	static INPUTCORE_API const FKey MixedReality_Left_Trigger_Axis;
+	static INPUTCORE_API const FKey MixedReality_Left_Thumbstick_2D;
+	static INPUTCORE_API const FKey MixedReality_Left_Thumbstick_X;
+	static INPUTCORE_API const FKey MixedReality_Left_Thumbstick_Y;
+	static INPUTCORE_API const FKey MixedReality_Left_Thumbstick_Click;
+	static INPUTCORE_API const FKey MixedReality_Left_Thumbstick_Up;
+	static INPUTCORE_API const FKey MixedReality_Left_Thumbstick_Down;
+	static INPUTCORE_API const FKey MixedReality_Left_Thumbstick_Left;
+	static INPUTCORE_API const FKey MixedReality_Left_Thumbstick_Right;
+	static INPUTCORE_API const FKey MixedReality_Left_Trackpad_2D;
+	static INPUTCORE_API const FKey MixedReality_Left_Trackpad_X;
+	static INPUTCORE_API const FKey MixedReality_Left_Trackpad_Y;
+	static INPUTCORE_API const FKey MixedReality_Left_Trackpad_Click;
+	static INPUTCORE_API const FKey MixedReality_Left_Trackpad_Touch;
+	static INPUTCORE_API const FKey MixedReality_Left_Trackpad_Up;
+	static INPUTCORE_API const FKey MixedReality_Left_Trackpad_Down;
+	static INPUTCORE_API const FKey MixedReality_Left_Trackpad_Left;
+	static INPUTCORE_API const FKey MixedReality_Left_Trackpad_Right;
+	static INPUTCORE_API const FKey MixedReality_Right_Menu_Click;
+	static INPUTCORE_API const FKey MixedReality_Right_Grip_Click;
+	static INPUTCORE_API const FKey MixedReality_Right_Trigger_Click;
+	static INPUTCORE_API const FKey MixedReality_Right_Trigger_Axis;
+	static INPUTCORE_API const FKey MixedReality_Right_Thumbstick_2D;
+	static INPUTCORE_API const FKey MixedReality_Right_Thumbstick_X;
+	static INPUTCORE_API const FKey MixedReality_Right_Thumbstick_Y;
+	static INPUTCORE_API const FKey MixedReality_Right_Thumbstick_Click;
+	static INPUTCORE_API const FKey MixedReality_Right_Thumbstick_Up;
+	static INPUTCORE_API const FKey MixedReality_Right_Thumbstick_Down;
+	static INPUTCORE_API const FKey MixedReality_Right_Thumbstick_Left;
+	static INPUTCORE_API const FKey MixedReality_Right_Thumbstick_Right;
+	static INPUTCORE_API const FKey MixedReality_Right_Trackpad_2D;
+	static INPUTCORE_API const FKey MixedReality_Right_Trackpad_X;
+	static INPUTCORE_API const FKey MixedReality_Right_Trackpad_Y;
+	static INPUTCORE_API const FKey MixedReality_Right_Trackpad_Click;
+	static INPUTCORE_API const FKey MixedReality_Right_Trackpad_Touch;
+	static INPUTCORE_API const FKey MixedReality_Right_Trackpad_Up;
+	static INPUTCORE_API const FKey MixedReality_Right_Trackpad_Down;
+	static INPUTCORE_API const FKey MixedReality_Right_Trackpad_Left;
+	static INPUTCORE_API const FKey MixedReality_Right_Trackpad_Right;
 
 	// Oculus Touch Controller
-	static const FKey OculusTouch_Left_X_Click;
-	static const FKey OculusTouch_Left_Y_Click;
-	static const FKey OculusTouch_Left_X_Touch;
-	static const FKey OculusTouch_Left_Y_Touch;
-	static const FKey OculusTouch_Left_Menu_Click;
-	static const FKey OculusTouch_Left_Grip_Click;
-	static const FKey OculusTouch_Left_Grip_Axis;
-	static const FKey OculusTouch_Left_Trigger_Click;
-	static const FKey OculusTouch_Left_Trigger_Axis;
-	static const FKey OculusTouch_Left_Trigger_Touch;
-	static const FKey OculusTouch_Left_Thumbstick_2D;
-	static const FKey OculusTouch_Left_Thumbstick_X;
-	static const FKey OculusTouch_Left_Thumbstick_Y;
-	static const FKey OculusTouch_Left_Thumbstick_Click;
-	static const FKey OculusTouch_Left_Thumbstick_Touch;
-	static const FKey OculusTouch_Left_Thumbstick_Up;
-	static const FKey OculusTouch_Left_Thumbstick_Down;
-	static const FKey OculusTouch_Left_Thumbstick_Left;
-	static const FKey OculusTouch_Left_Thumbstick_Right;
-	static const FKey OculusTouch_Right_A_Click;
-	static const FKey OculusTouch_Right_B_Click;
-	static const FKey OculusTouch_Right_A_Touch;
-	static const FKey OculusTouch_Right_B_Touch;
+	static INPUTCORE_API const FKey OculusTouch_Left_X_Click;
+	static INPUTCORE_API const FKey OculusTouch_Left_Y_Click;
+	static INPUTCORE_API const FKey OculusTouch_Left_X_Touch;
+	static INPUTCORE_API const FKey OculusTouch_Left_Y_Touch;
+	static INPUTCORE_API const FKey OculusTouch_Left_Menu_Click;
+	static INPUTCORE_API const FKey OculusTouch_Left_Grip_Click;
+	static INPUTCORE_API const FKey OculusTouch_Left_Grip_Axis;
+	static INPUTCORE_API const FKey OculusTouch_Left_Trigger_Click;
+	static INPUTCORE_API const FKey OculusTouch_Left_Trigger_Axis;
+	static INPUTCORE_API const FKey OculusTouch_Left_Trigger_Touch;
+	static INPUTCORE_API const FKey OculusTouch_Left_Thumbstick_2D;
+	static INPUTCORE_API const FKey OculusTouch_Left_Thumbstick_X;
+	static INPUTCORE_API const FKey OculusTouch_Left_Thumbstick_Y;
+	static INPUTCORE_API const FKey OculusTouch_Left_Thumbstick_Click;
+	static INPUTCORE_API const FKey OculusTouch_Left_Thumbstick_Touch;
+	static INPUTCORE_API const FKey OculusTouch_Left_Thumbstick_Up;
+	static INPUTCORE_API const FKey OculusTouch_Left_Thumbstick_Down;
+	static INPUTCORE_API const FKey OculusTouch_Left_Thumbstick_Left;
+	static INPUTCORE_API const FKey OculusTouch_Left_Thumbstick_Right;
+	static INPUTCORE_API const FKey OculusTouch_Right_A_Click;
+	static INPUTCORE_API const FKey OculusTouch_Right_B_Click;
+	static INPUTCORE_API const FKey OculusTouch_Right_A_Touch;
+	static INPUTCORE_API const FKey OculusTouch_Right_B_Touch;
 	UE_DEPRECATED(5.1, "This key has been deprecated and will be removed.")
-	static const FKey OculusTouch_Right_System_Click;
-	static const FKey OculusTouch_Right_Grip_Click;
-	static const FKey OculusTouch_Right_Grip_Axis;
-	static const FKey OculusTouch_Right_Trigger_Click;
-	static const FKey OculusTouch_Right_Trigger_Axis;
-	static const FKey OculusTouch_Right_Trigger_Touch;
-	static const FKey OculusTouch_Right_Thumbstick_2D;
-	static const FKey OculusTouch_Right_Thumbstick_X;
-	static const FKey OculusTouch_Right_Thumbstick_Y;
-	static const FKey OculusTouch_Right_Thumbstick_Click;
-	static const FKey OculusTouch_Right_Thumbstick_Touch;
-	static const FKey OculusTouch_Right_Thumbstick_Up;
-	static const FKey OculusTouch_Right_Thumbstick_Down;
-	static const FKey OculusTouch_Right_Thumbstick_Left;
-	static const FKey OculusTouch_Right_Thumbstick_Right;
+	static INPUTCORE_API const FKey OculusTouch_Right_System_Click;
+	static INPUTCORE_API const FKey OculusTouch_Right_Grip_Click;
+	static INPUTCORE_API const FKey OculusTouch_Right_Grip_Axis;
+	static INPUTCORE_API const FKey OculusTouch_Right_Trigger_Click;
+	static INPUTCORE_API const FKey OculusTouch_Right_Trigger_Axis;
+	static INPUTCORE_API const FKey OculusTouch_Right_Trigger_Touch;
+	static INPUTCORE_API const FKey OculusTouch_Right_Thumbstick_2D;
+	static INPUTCORE_API const FKey OculusTouch_Right_Thumbstick_X;
+	static INPUTCORE_API const FKey OculusTouch_Right_Thumbstick_Y;
+	static INPUTCORE_API const FKey OculusTouch_Right_Thumbstick_Click;
+	static INPUTCORE_API const FKey OculusTouch_Right_Thumbstick_Touch;
+	static INPUTCORE_API const FKey OculusTouch_Right_Thumbstick_Up;
+	static INPUTCORE_API const FKey OculusTouch_Right_Thumbstick_Down;
+	static INPUTCORE_API const FKey OculusTouch_Right_Thumbstick_Left;
+	static INPUTCORE_API const FKey OculusTouch_Right_Thumbstick_Right;
 
 	// Valve Index Controller
-	static const FKey ValveIndex_Left_A_Click;
-	static const FKey ValveIndex_Left_B_Click;
-	static const FKey ValveIndex_Left_A_Touch;
-	static const FKey ValveIndex_Left_B_Touch;
+	static INPUTCORE_API const FKey ValveIndex_Left_A_Click;
+	static INPUTCORE_API const FKey ValveIndex_Left_B_Click;
+	static INPUTCORE_API const FKey ValveIndex_Left_A_Touch;
+	static INPUTCORE_API const FKey ValveIndex_Left_B_Touch;
 	UE_DEPRECATED(5.1, "This key has been deprecated and will be removed.")
-	static const FKey ValveIndex_Left_System_Click;
+	static INPUTCORE_API const FKey ValveIndex_Left_System_Click;
 	UE_DEPRECATED(5.1, "This key has been deprecated and will be removed.")
-	static const FKey ValveIndex_Left_System_Touch;
-	static const FKey ValveIndex_Left_Grip_Axis;
-	static const FKey ValveIndex_Left_Grip_Force;
-	static const FKey ValveIndex_Left_Trigger_Click;
-	static const FKey ValveIndex_Left_Trigger_Axis;
-	static const FKey ValveIndex_Left_Trigger_Touch;
-	static const FKey ValveIndex_Left_Thumbstick_2D;
-	static const FKey ValveIndex_Left_Thumbstick_X;
-	static const FKey ValveIndex_Left_Thumbstick_Y;
-	static const FKey ValveIndex_Left_Thumbstick_Click;
-	static const FKey ValveIndex_Left_Thumbstick_Touch;
-	static const FKey ValveIndex_Left_Thumbstick_Up;
-	static const FKey ValveIndex_Left_Thumbstick_Down;
-	static const FKey ValveIndex_Left_Thumbstick_Left;
-	static const FKey ValveIndex_Left_Thumbstick_Right;
-	static const FKey ValveIndex_Left_Trackpad_2D;
-	static const FKey ValveIndex_Left_Trackpad_X;
-	static const FKey ValveIndex_Left_Trackpad_Y;
-	static const FKey ValveIndex_Left_Trackpad_Force;
-	static const FKey ValveIndex_Left_Trackpad_Touch;
-	static const FKey ValveIndex_Left_Trackpad_Up;
-	static const FKey ValveIndex_Left_Trackpad_Down;
-	static const FKey ValveIndex_Left_Trackpad_Left;
-	static const FKey ValveIndex_Left_Trackpad_Right;
-	static const FKey ValveIndex_Right_A_Click;
-	static const FKey ValveIndex_Right_B_Click;
-	static const FKey ValveIndex_Right_A_Touch;
-	static const FKey ValveIndex_Right_B_Touch;
+	static INPUTCORE_API const FKey ValveIndex_Left_System_Touch;
+	static INPUTCORE_API const FKey ValveIndex_Left_Grip_Axis;
+	static INPUTCORE_API const FKey ValveIndex_Left_Grip_Force;
+	static INPUTCORE_API const FKey ValveIndex_Left_Trigger_Click;
+	static INPUTCORE_API const FKey ValveIndex_Left_Trigger_Axis;
+	static INPUTCORE_API const FKey ValveIndex_Left_Trigger_Touch;
+	static INPUTCORE_API const FKey ValveIndex_Left_Thumbstick_2D;
+	static INPUTCORE_API const FKey ValveIndex_Left_Thumbstick_X;
+	static INPUTCORE_API const FKey ValveIndex_Left_Thumbstick_Y;
+	static INPUTCORE_API const FKey ValveIndex_Left_Thumbstick_Click;
+	static INPUTCORE_API const FKey ValveIndex_Left_Thumbstick_Touch;
+	static INPUTCORE_API const FKey ValveIndex_Left_Thumbstick_Up;
+	static INPUTCORE_API const FKey ValveIndex_Left_Thumbstick_Down;
+	static INPUTCORE_API const FKey ValveIndex_Left_Thumbstick_Left;
+	static INPUTCORE_API const FKey ValveIndex_Left_Thumbstick_Right;
+	static INPUTCORE_API const FKey ValveIndex_Left_Trackpad_2D;
+	static INPUTCORE_API const FKey ValveIndex_Left_Trackpad_X;
+	static INPUTCORE_API const FKey ValveIndex_Left_Trackpad_Y;
+	static INPUTCORE_API const FKey ValveIndex_Left_Trackpad_Force;
+	static INPUTCORE_API const FKey ValveIndex_Left_Trackpad_Touch;
+	static INPUTCORE_API const FKey ValveIndex_Left_Trackpad_Up;
+	static INPUTCORE_API const FKey ValveIndex_Left_Trackpad_Down;
+	static INPUTCORE_API const FKey ValveIndex_Left_Trackpad_Left;
+	static INPUTCORE_API const FKey ValveIndex_Left_Trackpad_Right;
+	static INPUTCORE_API const FKey ValveIndex_Right_A_Click;
+	static INPUTCORE_API const FKey ValveIndex_Right_B_Click;
+	static INPUTCORE_API const FKey ValveIndex_Right_A_Touch;
+	static INPUTCORE_API const FKey ValveIndex_Right_B_Touch;
 	UE_DEPRECATED(5.1, "This key has been deprecated and will be removed.")
-	static const FKey ValveIndex_Right_System_Click;
+	static INPUTCORE_API const FKey ValveIndex_Right_System_Click;
 	UE_DEPRECATED(5.1, "This key has been deprecated and will be removed.")
-	static const FKey ValveIndex_Right_System_Touch;
-	static const FKey ValveIndex_Right_Grip_Axis;
-	static const FKey ValveIndex_Right_Grip_Force;
-	static const FKey ValveIndex_Right_Trigger_Click;
-	static const FKey ValveIndex_Right_Trigger_Axis;
-	static const FKey ValveIndex_Right_Trigger_Touch;
-	static const FKey ValveIndex_Right_Thumbstick_2D;
-	static const FKey ValveIndex_Right_Thumbstick_X;
-	static const FKey ValveIndex_Right_Thumbstick_Y;
-	static const FKey ValveIndex_Right_Thumbstick_Click;
-	static const FKey ValveIndex_Right_Thumbstick_Touch;
-	static const FKey ValveIndex_Right_Thumbstick_Up;
-	static const FKey ValveIndex_Right_Thumbstick_Down;
-	static const FKey ValveIndex_Right_Thumbstick_Left;
-	static const FKey ValveIndex_Right_Thumbstick_Right;
-	static const FKey ValveIndex_Right_Trackpad_2D;
-	static const FKey ValveIndex_Right_Trackpad_X;
-	static const FKey ValveIndex_Right_Trackpad_Y;
-	static const FKey ValveIndex_Right_Trackpad_Force;
-	static const FKey ValveIndex_Right_Trackpad_Touch;
-	static const FKey ValveIndex_Right_Trackpad_Up;
-	static const FKey ValveIndex_Right_Trackpad_Down;
-	static const FKey ValveIndex_Right_Trackpad_Left;
-	static const FKey ValveIndex_Right_Trackpad_Right;
+	static INPUTCORE_API const FKey ValveIndex_Right_System_Touch;
+	static INPUTCORE_API const FKey ValveIndex_Right_Grip_Axis;
+	static INPUTCORE_API const FKey ValveIndex_Right_Grip_Force;
+	static INPUTCORE_API const FKey ValveIndex_Right_Trigger_Click;
+	static INPUTCORE_API const FKey ValveIndex_Right_Trigger_Axis;
+	static INPUTCORE_API const FKey ValveIndex_Right_Trigger_Touch;
+	static INPUTCORE_API const FKey ValveIndex_Right_Thumbstick_2D;
+	static INPUTCORE_API const FKey ValveIndex_Right_Thumbstick_X;
+	static INPUTCORE_API const FKey ValveIndex_Right_Thumbstick_Y;
+	static INPUTCORE_API const FKey ValveIndex_Right_Thumbstick_Click;
+	static INPUTCORE_API const FKey ValveIndex_Right_Thumbstick_Touch;
+	static INPUTCORE_API const FKey ValveIndex_Right_Thumbstick_Up;
+	static INPUTCORE_API const FKey ValveIndex_Right_Thumbstick_Down;
+	static INPUTCORE_API const FKey ValveIndex_Right_Thumbstick_Left;
+	static INPUTCORE_API const FKey ValveIndex_Right_Thumbstick_Right;
+	static INPUTCORE_API const FKey ValveIndex_Right_Trackpad_2D;
+	static INPUTCORE_API const FKey ValveIndex_Right_Trackpad_X;
+	static INPUTCORE_API const FKey ValveIndex_Right_Trackpad_Y;
+	static INPUTCORE_API const FKey ValveIndex_Right_Trackpad_Force;
+	static INPUTCORE_API const FKey ValveIndex_Right_Trackpad_Touch;
+	static INPUTCORE_API const FKey ValveIndex_Right_Trackpad_Up;
+	static INPUTCORE_API const FKey ValveIndex_Right_Trackpad_Down;
+	static INPUTCORE_API const FKey ValveIndex_Right_Trackpad_Left;
+	static INPUTCORE_API const FKey ValveIndex_Right_Trackpad_Right;
 
 	// Virtual buttons that use other buttons depending on the platform
-	static const FKey Virtual_Accept;
-	static const FKey Virtual_Back;
+	static INPUTCORE_API const FKey Virtual_Accept;
+	static INPUTCORE_API const FKey Virtual_Back;
 
-	static const FKey Invalid;
+	static INPUTCORE_API const FKey Invalid;
 
 	static const int32 NUM_TOUCH_KEYS = 11;
-	static const FKey TouchKeys[NUM_TOUCH_KEYS];
+	static INPUTCORE_API const FKey TouchKeys[NUM_TOUCH_KEYS];
 
 	// XR key names are parseable into exactly 4 tokens
 	static const int32 NUM_XR_KEY_TOKENS = 4;
 
-	static EConsoleForGamepadLabels::Type ConsoleForGamepadLabels;
+	static INPUTCORE_API EConsoleForGamepadLabels::Type ConsoleForGamepadLabels;
 
-	static const FName NAME_KeyboardCategory;
-	static const FName NAME_GamepadCategory;
-	static const FName NAME_MouseCategory;
+	static INPUTCORE_API const FName NAME_KeyboardCategory;
+	static INPUTCORE_API const FName NAME_GamepadCategory;
+	static INPUTCORE_API const FName NAME_MouseCategory;
 
-	static void Initialize();
-	static void AddKey(const FKeyDetails& KeyDetails);
-	static void AddPairedKey(const FKeyDetails& PairedKeyDetails, FKey KeyX, FKey KeyY);	// Map the two provided keys to the X and Z axes of the paired key
-	static void GetAllKeys(TArray<FKey>& OutKeys);
-	static TSharedPtr<FKeyDetails> GetKeyDetails(const FKey Key);
-	static void RemoveKeysWithCategory(const FName InCategory);
+	static INPUTCORE_API void Initialize();
+	static INPUTCORE_API void AddKey(const FKeyDetails& KeyDetails);
+	static INPUTCORE_API void AddPairedKey(const FKeyDetails& PairedKeyDetails, FKey KeyX, FKey KeyY);	// Map the two provided keys to the X and Z axes of the paired key
+	static INPUTCORE_API void GetAllKeys(TArray<FKey>& OutKeys);
+	static INPUTCORE_API TSharedPtr<FKeyDetails> GetKeyDetails(const FKey Key);
+	static INPUTCORE_API void RemoveKeysWithCategory(const FName InCategory);
 
 	// These exist for backwards compatibility reasons only
 	static bool IsModifierKey(FKey Key) { return Key.IsModifierKey(); }
@@ -715,11 +715,11 @@ struct INPUTCORE_API EKeys
 	static void SetConsoleForGamepadLabels(const EConsoleForGamepadLabels::Type Console) { ConsoleForGamepadLabels = Console; }
 
 	// Function that provides remapping for some gamepad keys in display windows
-	static FText GetGamepadDisplayName(const FKey Key);
+	static INPUTCORE_API FText GetGamepadDisplayName(const FKey Key);
 
-	static void AddMenuCategoryDisplayInfo(const FName CategoryName, const FText DisplayName, const FName PaletteIcon);
-	static FText GetMenuCategoryDisplayName(const FName CategoryName);
-	static FName GetMenuCategoryPaletteIcon(const FName CategoryName);
+	static INPUTCORE_API void AddMenuCategoryDisplayInfo(const FName CategoryName, const FText DisplayName, const FName PaletteIcon);
+	static INPUTCORE_API FText GetMenuCategoryDisplayName(const FName CategoryName);
+	static INPUTCORE_API FName GetMenuCategoryPaletteIcon(const FName CategoryName);
 
 private:
 
@@ -729,9 +729,9 @@ private:
 		FName PaletteIcon;
 	};
 
-	static TMap<FKey, TSharedPtr<FKeyDetails> > InputKeys;
-	static TMap<FName, FCategoryDisplayInfo> MenuCategoryDisplayInfo;
-	static bool bInitialized;
+	static INPUTCORE_API TMap<FKey, TSharedPtr<FKeyDetails> > InputKeys;
+	static INPUTCORE_API TMap<FName, FCategoryDisplayInfo> MenuCategoryDisplayInfo;
+	static INPUTCORE_API bool bInitialized;
 
 };
 
@@ -753,19 +753,19 @@ namespace ETouchType
 }
 
 
-struct INPUTCORE_API FInputKeyManager
+struct FInputKeyManager
 {
 public:
-	static FInputKeyManager& Get();
+	static INPUTCORE_API FInputKeyManager& Get();
 
-	void GetCodesFromKey(const FKey Key, const uint32*& KeyCode, const uint32*& CharCode) const;
+	INPUTCORE_API void GetCodesFromKey(const FKey Key, const uint32*& KeyCode, const uint32*& CharCode) const;
 
 	/**
 	 * Retrieves the key mapped to the specified character code.
 	 * @param KeyCode	The key code to get the name for.
 	 */
-	FKey GetKeyFromCodes( const uint32 KeyCode, const uint32 CharCode ) const;
-	void InitKeyMappings();
+	INPUTCORE_API FKey GetKeyFromCodes( const uint32 KeyCode, const uint32 CharCode ) const;
+	INPUTCORE_API void InitKeyMappings();
 private:
 	FInputKeyManager()
 	{
diff --git a/Engine/Source/Runtime/InputCore/Public/Android/AndroidPlatformInput.h b/Engine/Source/Runtime/InputCore/Public/Android/AndroidPlatformInput.h
--- a/Engine/Source/Runtime/InputCore/Public/Android/AndroidPlatformInput.h
+++ b/Engine/Source/Runtime/InputCore/Public/Android/AndroidPlatformInput.h
@@ -5,10 +5,10 @@
 #include "CoreTypes.h"
 #include "GenericPlatform/GenericPlatformInput.h"
 
-struct INPUTCORE_API FAndroidPlatformInput : FGenericPlatformInput
+struct FAndroidPlatformInput : FGenericPlatformInput
 {
-	static uint32 GetCharKeyMap(uint32* KeyCodes, FString* KeyNames, uint32 MaxMappings);
-	static uint32 GetKeyMap( uint32* KeyCodes, FString* KeyNames, uint32 MaxMappings );
+	static INPUTCORE_API uint32 GetCharKeyMap(uint32* KeyCodes, FString* KeyNames, uint32 MaxMappings);
+	static INPUTCORE_API uint32 GetKeyMap( uint32* KeyCodes, FString* KeyNames, uint32 MaxMappings );
 };
 
 typedef FAndroidPlatformInput FPlatformInput;
diff --git a/Engine/Source/Runtime/InputCore/Public/GenericPlatform/GenericPlatformInput.h b/Engine/Source/Runtime/InputCore/Public/GenericPlatform/GenericPlatformInput.h
--- a/Engine/Source/Runtime/InputCore/Public/GenericPlatform/GenericPlatformInput.h
+++ b/Engine/Source/Runtime/InputCore/Public/GenericPlatform/GenericPlatformInput.h
@@ -8,7 +8,7 @@
 
 class FString;
 
-struct INPUTCORE_API FGenericPlatformInput
+struct FGenericPlatformInput
 {
 public:
 	FORCEINLINE static uint32 GetKeyMap( uint32* KeyCodes, FString* KeyNames, uint32 MaxMappings )
@@ -42,5 +42,5 @@ protected:
 	* @param bMapUppercaseKeys If true, will map A, B, C, etc to EKeys::A, EKeys::B, EKeys::C
 	* @param bMapLowercaseKeys If true, will map a, b, c, etc to EKeys::A, EKeys::B, EKeys::C
 	*/
-	static uint32 GetStandardPrintableKeyMap(uint32* KeyCodes, FString* KeyNames, uint32 MaxMappings, bool bMapUppercaseKeys, bool bMapLowercaseKeys);
+	static INPUTCORE_API uint32 GetStandardPrintableKeyMap(uint32* KeyCodes, FString* KeyNames, uint32 MaxMappings, bool bMapUppercaseKeys, bool bMapLowercaseKeys);
 };
diff --git a/Engine/Source/Runtime/InputCore/Public/Linux/LinuxPlatformInput.h b/Engine/Source/Runtime/InputCore/Public/Linux/LinuxPlatformInput.h
--- a/Engine/Source/Runtime/InputCore/Public/Linux/LinuxPlatformInput.h
+++ b/Engine/Source/Runtime/InputCore/Public/Linux/LinuxPlatformInput.h
@@ -6,10 +6,10 @@
 
 class FString;
 
-struct INPUTCORE_API FLinuxPlatformInput : FGenericPlatformInput
+struct FLinuxPlatformInput : FGenericPlatformInput
 {
-	static uint32 GetKeyMap( uint32* KeyCodes, FString* KeyNames, uint32 MaxMappings );
-	static uint32 GetCharKeyMap(uint32* KeyCodes, FString* KeyNames, uint32 MaxMappings);
+	static INPUTCORE_API uint32 GetKeyMap( uint32* KeyCodes, FString* KeyNames, uint32 MaxMappings );
+	static INPUTCORE_API uint32 GetCharKeyMap(uint32* KeyCodes, FString* KeyNames, uint32 MaxMappings);
 };
 
 typedef FLinuxPlatformInput FPlatformInput;
diff --git a/Engine/Source/Runtime/InputCore/Public/Windows/WindowsPlatformInput.h b/Engine/Source/Runtime/InputCore/Public/Windows/WindowsPlatformInput.h
--- a/Engine/Source/Runtime/InputCore/Public/Windows/WindowsPlatformInput.h
+++ b/Engine/Source/Runtime/InputCore/Public/Windows/WindowsPlatformInput.h
@@ -5,10 +5,10 @@
 #include "CoreTypes.h"
 #include "GenericPlatform/GenericPlatformInput.h"
 
-struct INPUTCORE_API FWindowsPlatformInput : FGenericPlatformInput
+struct FWindowsPlatformInput : FGenericPlatformInput
 {
-	static uint32 GetKeyMap( uint32* KeyCodes, FString* KeyNames, uint32 MaxMappings );
-	static uint32 GetCharKeyMap(uint32* KeyCodes, FString* KeyNames, uint32 MaxMappings);
+	static INPUTCORE_API uint32 GetKeyMap( uint32* KeyCodes, FString* KeyNames, uint32 MaxMappings );
+	static INPUTCORE_API uint32 GetCharKeyMap(uint32* KeyCodes, FString* KeyNames, uint32 MaxMappings);
 };
 
 typedef FWindowsPlatformInput FPlatformInput;
diff --git a/Engine/Source/Runtime/InstallBundleManager/Public/BundlePrereqCombinedStatusHelper.h b/Engine/Source/Runtime/InstallBundleManager/Public/BundlePrereqCombinedStatusHelper.h
--- a/Engine/Source/Runtime/InstallBundleManager/Public/BundlePrereqCombinedStatusHelper.h
+++ b/Engine/Source/Runtime/InstallBundleManager/Public/BundlePrereqCombinedStatusHelper.h
@@ -22,7 +22,7 @@ struct FInstallBundleRequestResultInfo;
 //Handles calculating the bundle status by combining progress from all of its
 //Prerequisites. Allows you to display one progress percent that is weighted based on all
 //bundles' values.
-class INSTALLBUNDLEMANAGER_API FInstallBundleCombinedProgressTracker
+class FInstallBundleCombinedProgressTracker
 {
 public:
 
@@ -50,37 +50,37 @@ public:
 	
 public:
 	/** bAutoTick Whehter this tracker should automatically Tick */
-	FInstallBundleCombinedProgressTracker(bool bAutoTick = true);
-	~FInstallBundleCombinedProgressTracker();
+	INSTALLBUNDLEMANAGER_API FInstallBundleCombinedProgressTracker(bool bAutoTick = true);
+	INSTALLBUNDLEMANAGER_API ~FInstallBundleCombinedProgressTracker();
 	
-	FInstallBundleCombinedProgressTracker(const FInstallBundleCombinedProgressTracker& Other);
-	FInstallBundleCombinedProgressTracker(FInstallBundleCombinedProgressTracker&& Other);
+	INSTALLBUNDLEMANAGER_API FInstallBundleCombinedProgressTracker(const FInstallBundleCombinedProgressTracker& Other);
+	INSTALLBUNDLEMANAGER_API FInstallBundleCombinedProgressTracker(FInstallBundleCombinedProgressTracker&& Other);
 	
-	FInstallBundleCombinedProgressTracker& operator=(const FInstallBundleCombinedProgressTracker& Other);
-	FInstallBundleCombinedProgressTracker& operator=(FInstallBundleCombinedProgressTracker&& Other);
+	INSTALLBUNDLEMANAGER_API FInstallBundleCombinedProgressTracker& operator=(const FInstallBundleCombinedProgressTracker& Other);
+	INSTALLBUNDLEMANAGER_API FInstallBundleCombinedProgressTracker& operator=(FInstallBundleCombinedProgressTracker&& Other);
 	
 	//Setup tracking for all bundles required in the supplied BundleContentState
-	void SetBundlesToTrackFromContentState(const FInstallBundleCombinedContentState& BundleContentState, TArrayView<FName> BundlesToTrack);
+	INSTALLBUNDLEMANAGER_API void SetBundlesToTrackFromContentState(const FInstallBundleCombinedContentState& BundleContentState, TArrayView<FName> BundlesToTrack);
 	
 	//Get current CombinedBundleStatus for everything setup to track
-	const FCombinedProgress& GetCurrentCombinedProgress() const;
+	INSTALLBUNDLEMANAGER_API const FCombinedProgress& GetCurrentCombinedProgress() const;
 	
 	//Useful for resolving tick order issue
 	void ForceTick() { Tick(0); }
 
 private:
-	bool Tick(float dt);
-	void UpdateBundleCache();
-	void UpdateCombinedStatus();
+	INSTALLBUNDLEMANAGER_API bool Tick(float dt);
+	INSTALLBUNDLEMANAGER_API void UpdateBundleCache();
+	INSTALLBUNDLEMANAGER_API void UpdateCombinedStatus();
 	
-	void SetupDelegates(bool bAutoTick);
-	void CleanUpDelegates();
+	INSTALLBUNDLEMANAGER_API void SetupDelegates(bool bAutoTick);
+	INSTALLBUNDLEMANAGER_API void CleanUpDelegates();
 	
 	//Called so we can track when a bundle is finished
-	void OnBundleInstallComplete(FInstallBundleRequestResultInfo CompletedBundleInfo);
-	void OnBundleInstallPauseChanged(FInstallBundlePauseInfo PauseInfo);
+	INSTALLBUNDLEMANAGER_API void OnBundleInstallComplete(FInstallBundleRequestResultInfo CompletedBundleInfo);
+	INSTALLBUNDLEMANAGER_API void OnBundleInstallPauseChanged(FInstallBundlePauseInfo PauseInfo);
 	
-	float GetCombinedProgressPercent() const;
+	INSTALLBUNDLEMANAGER_API float GetCombinedProgressPercent() const;
 	
 private:
 	//All bundles we need including pre-reqs
diff --git a/Engine/Source/Runtime/InstallBundleManager/Public/InstallBundleCache.h b/Engine/Source/Runtime/InstallBundleManager/Public/InstallBundleCache.h
--- a/Engine/Source/Runtime/InstallBundleManager/Public/InstallBundleCache.h
+++ b/Engine/Source/Runtime/InstallBundleManager/Public/InstallBundleCache.h
@@ -49,58 +49,58 @@ struct FInstallBundleCacheFlushResult
 	TMap<FName, TArray<EInstallBundleSourceType>> BundlesToEvict;
 };
 
-class INSTALLBUNDLEMANAGER_API FInstallBundleCache : public TSharedFromThis<FInstallBundleCache>
+class FInstallBundleCache : public TSharedFromThis<FInstallBundleCache>
 {
 public:
-	virtual ~FInstallBundleCache();
+	INSTALLBUNDLEMANAGER_API virtual ~FInstallBundleCache();
 
-	void Init(FInstallBundleCacheInitInfo InitInfo);
+	INSTALLBUNDLEMANAGER_API void Init(FInstallBundleCacheInitInfo InitInfo);
 
 	FName GetName() const { return CacheName; }
 
 	// Add a bundle to the cache.  
-	void AddOrUpdateBundle(EInstallBundleSourceType Source, const FInstallBundleCacheBundleInfo& AddInfo);
+	INSTALLBUNDLEMANAGER_API void AddOrUpdateBundle(EInstallBundleSourceType Source, const FInstallBundleCacheBundleInfo& AddInfo);
 
-	void RemoveBundle(EInstallBundleSourceType Source, FName BundleName);
+	INSTALLBUNDLEMANAGER_API void RemoveBundle(EInstallBundleSourceType Source, FName BundleName);
 
-	TOptional<FInstallBundleCacheBundleInfo> GetBundleInfo(FName BundleName) const;
-	TOptional<FInstallBundleCacheBundleInfo> GetBundleInfo(EInstallBundleSourceType Source, FName BundleName) const;	
+	INSTALLBUNDLEMANAGER_API TOptional<FInstallBundleCacheBundleInfo> GetBundleInfo(FName BundleName) const;
+	INSTALLBUNDLEMANAGER_API TOptional<FInstallBundleCacheBundleInfo> GetBundleInfo(EInstallBundleSourceType Source, FName BundleName) const;	
 
 	// Return the total size of the cache
-	uint64 GetSize() const;
+	INSTALLBUNDLEMANAGER_API uint64 GetSize() const;
 	// Return the amount of space in use.  This could possbly exceed GetSize() if the cache size is changed or more 
 	// bundles are added the to cache.
-	uint64 GetUsedSize() const;
+	INSTALLBUNDLEMANAGER_API uint64 GetUsedSize() const;
 	// Return the amount of free space in the cache, clamped to [0, GetSize()]
-	uint64 GetFreeSpace() const;
+	INSTALLBUNDLEMANAGER_API uint64 GetFreeSpace() const;
 
 	// Called from bundle manager
-	FInstallBundleCacheReserveResult Reserve(FName BundleName);
+	INSTALLBUNDLEMANAGER_API FInstallBundleCacheReserveResult Reserve(FName BundleName);
 
 	// Called from bundle manager, returns all bundles that can be evicted
-	FInstallBundleCacheFlushResult Flush(EInstallBundleSourceType* Source = nullptr);
+	INSTALLBUNDLEMANAGER_API FInstallBundleCacheFlushResult Flush(EInstallBundleSourceType* Source = nullptr);
 
-	bool Contains(FName BundleName) const;
-	bool Contains(EInstallBundleSourceType Source, FName BundleName) const;
+	INSTALLBUNDLEMANAGER_API bool Contains(FName BundleName) const;
+	INSTALLBUNDLEMANAGER_API bool Contains(EInstallBundleSourceType Source, FName BundleName) const;
 
 	// Called from bundle manager to make the files for this bundle eligible for eviction
-	bool Release(FName BundleName);
+	INSTALLBUNDLEMANAGER_API bool Release(FName BundleName);
 
-	bool SetPendingEvict(FName BundleName);
+	INSTALLBUNDLEMANAGER_API bool SetPendingEvict(FName BundleName);
 
-	bool ClearPendingEvict(FName BundleName);
+	INSTALLBUNDLEMANAGER_API bool ClearPendingEvict(FName BundleName);
 
 	// Hint to the cache that this bundle is requested, and we should prefer to evict non-requested bundles if possible
-	void HintRequested(FName BundleName, bool bRequested);
+	INSTALLBUNDLEMANAGER_API void HintRequested(FName BundleName, bool bRequested);
 
-	FInstallBundleCacheStats GetStats(EInstallBundleCacheDumpToLog DumpToLog = EInstallBundleCacheDumpToLog::None, bool bVerbose = false) const;
+	INSTALLBUNDLEMANAGER_API FInstallBundleCacheStats GetStats(EInstallBundleCacheDumpToLog DumpToLog = EInstallBundleCacheDumpToLog::None, bool bVerbose = false) const;
 
 private:
-	uint64 GetFreeSpaceInternal(uint64 UsedSize) const;
+	INSTALLBUNDLEMANAGER_API uint64 GetFreeSpaceInternal(uint64 UsedSize) const;
 
-	void CheckInvariants() const;
+	INSTALLBUNDLEMANAGER_API void CheckInvariants() const;
 	
-	void UpdateCacheInfoFromSourceInfo(FName BundleName);
+	INSTALLBUNDLEMANAGER_API void UpdateCacheInfoFromSourceInfo(FName BundleName);
 
 private:
 	struct FPerSourceBundleCacheInfo
diff --git a/Engine/Source/Runtime/InstallBundleManager/Public/InstallBundleManagerInterface.h b/Engine/Source/Runtime/InstallBundleManager/Public/InstallBundleManagerInterface.h
--- a/Engine/Source/Runtime/InstallBundleManager/Public/InstallBundleManagerInterface.h
+++ b/Engine/Source/Runtime/InstallBundleManager/Public/InstallBundleManagerInterface.h
@@ -91,17 +91,17 @@ DECLARE_DELEGATE_OneParam(FInstallBundleGetInstallStateDelegate, FInstallBundleC
 
 DECLARE_DELEGATE(FInstallBundleManagerFlushCacheCompleteDelegate);
 
-class INSTALLBUNDLEMANAGER_API IInstallBundleManager : public TSharedFromThis<IInstallBundleManager>
+class IInstallBundleManager : public TSharedFromThis<IInstallBundleManager>
 {
 public:
-	static FInstallBundleManagerInitCompleteMultiDelegate InitCompleteDelegate;
+	static INSTALLBUNDLEMANAGER_API FInstallBundleManagerInitCompleteMultiDelegate InitCompleteDelegate;
 
-	static FInstallBundleCompleteMultiDelegate InstallBundleCompleteDelegate; // Called when a content request is complete
-	static FInstallBundlePausedMultiDelegate PausedBundleDelegate;
-	static FInstallBundleReleasedMultiDelegate ReleasedDelegate; // Called when content release request is complete
-	static FInstallBundleManagerOnPatchCheckComplete PatchCheckCompleteDelegate;
+	static INSTALLBUNDLEMANAGER_API FInstallBundleCompleteMultiDelegate InstallBundleCompleteDelegate; // Called when a content request is complete
+	static INSTALLBUNDLEMANAGER_API FInstallBundlePausedMultiDelegate PausedBundleDelegate;
+	static INSTALLBUNDLEMANAGER_API FInstallBundleReleasedMultiDelegate ReleasedDelegate; // Called when content release request is complete
+	static INSTALLBUNDLEMANAGER_API FInstallBundleManagerOnPatchCheckComplete PatchCheckCompleteDelegate;
 
-	static TSharedPtr<IInstallBundleManager> GetPlatformInstallBundleManager();
+	static INSTALLBUNDLEMANAGER_API TSharedPtr<IInstallBundleManager> GetPlatformInstallBundleManager();
 
 	virtual ~IInstallBundleManager() {}
 
@@ -109,7 +109,7 @@ public:
 
 	virtual bool HasBundleSource(EInstallBundleSourceType SourceType) const = 0;
 
-	virtual const TSharedPtr<IInstallBundleSource> GetBundleSource(EInstallBundleSourceType SourceType) const;
+	INSTALLBUNDLEMANAGER_API virtual const TSharedPtr<IInstallBundleSource> GetBundleSource(EInstallBundleSourceType SourceType) const;
 
 
 	virtual FDelegateHandle PushInitErrorCallback(FInstallBundleManagerInitErrorHandler Callback) = 0;
@@ -119,45 +119,45 @@ public:
 
 	virtual EInstallBundleManagerInitState GetInitState() const = 0;
 
-	TValueOrError<FInstallBundleRequestInfo, EInstallBundleResult> RequestUpdateContent(FName BundleName, EInstallBundleRequestFlags Flags, ELogVerbosity::Type LogVerbosityOverride = ELogVerbosity::NoLogging, InstallBundleUtil::FContentRequestSharedContextPtr RequestSharedContext = nullptr);
+	INSTALLBUNDLEMANAGER_API TValueOrError<FInstallBundleRequestInfo, EInstallBundleResult> RequestUpdateContent(FName BundleName, EInstallBundleRequestFlags Flags, ELogVerbosity::Type LogVerbosityOverride = ELogVerbosity::NoLogging, InstallBundleUtil::FContentRequestSharedContextPtr RequestSharedContext = nullptr);
 	virtual TValueOrError<FInstallBundleRequestInfo, EInstallBundleResult> RequestUpdateContent(TArrayView<const FName> BundleNames, EInstallBundleRequestFlags Flags, ELogVerbosity::Type LogVerbosityOverride = ELogVerbosity::NoLogging, InstallBundleUtil::FContentRequestSharedContextPtr RequestSharedContext = nullptr) = 0;
 
-	FDelegateHandle GetContentState(FName BundleName, EInstallBundleGetContentStateFlags Flags, bool bAddDependencies, FInstallBundleGetContentStateDelegate Callback, FName RequestTag = NAME_None);
+	INSTALLBUNDLEMANAGER_API FDelegateHandle GetContentState(FName BundleName, EInstallBundleGetContentStateFlags Flags, bool bAddDependencies, FInstallBundleGetContentStateDelegate Callback, FName RequestTag = NAME_None);
 	virtual FDelegateHandle GetContentState(TArrayView<const FName> BundleNames, EInstallBundleGetContentStateFlags Flags, bool bAddDependencies, FInstallBundleGetContentStateDelegate Callback, FName RequestTag = NAME_None) = 0;
 	virtual void CancelAllGetContentStateRequestsForTag(FName RequestTag) = 0;
 	virtual void CancelAllGetContentStateRequests(FDelegateHandle Handle) = 0;
 
 	// Less expensive version of GetContentState() that only returns install state
 	// Synchronous versions return null if bundle manager is not yet initialized
-	FDelegateHandle GetInstallState(FName BundleName, bool bAddDependencies, FInstallBundleGetInstallStateDelegate Callback, FName RequestTag = NAME_None);
+	INSTALLBUNDLEMANAGER_API FDelegateHandle GetInstallState(FName BundleName, bool bAddDependencies, FInstallBundleGetInstallStateDelegate Callback, FName RequestTag = NAME_None);
 	virtual FDelegateHandle GetInstallState(TArrayView<const FName> BundleNames, bool bAddDependencies, FInstallBundleGetInstallStateDelegate Callback, FName RequestTag = NAME_None) = 0;
-	TValueOrError<FInstallBundleCombinedInstallState, EInstallBundleResult> GetInstallStateSynchronous(FName BundleName, bool bAddDependencies) const;
+	INSTALLBUNDLEMANAGER_API TValueOrError<FInstallBundleCombinedInstallState, EInstallBundleResult> GetInstallStateSynchronous(FName BundleName, bool bAddDependencies) const;
 	virtual TValueOrError<FInstallBundleCombinedInstallState, EInstallBundleResult> GetInstallStateSynchronous(TArrayView<const FName> BundleNames, bool bAddDependencies) const = 0;
 	virtual void CancelAllGetInstallStateRequestsForTag(FName RequestTag) = 0;
 	virtual void CancelAllGetInstallStateRequests(FDelegateHandle Handle) = 0;
 
-	TValueOrError<FInstallBundleRequestInfo, EInstallBundleResult> RequestReleaseContent(FName ReleaseName, EInstallBundleReleaseRequestFlags Flags, TArrayView<const FName> KeepNames = TArrayView<const FName>(), ELogVerbosity::Type LogVerbosityOverride = ELogVerbosity::NoLogging);
+	INSTALLBUNDLEMANAGER_API TValueOrError<FInstallBundleRequestInfo, EInstallBundleResult> RequestReleaseContent(FName ReleaseName, EInstallBundleReleaseRequestFlags Flags, TArrayView<const FName> KeepNames = TArrayView<const FName>(), ELogVerbosity::Type LogVerbosityOverride = ELogVerbosity::NoLogging);
 	virtual TValueOrError<FInstallBundleRequestInfo, EInstallBundleResult> RequestReleaseContent(TArrayView<const FName> ReleaseNames, EInstallBundleReleaseRequestFlags Flags, TArrayView<const FName> KeepNames = TArrayView<const FName>(), ELogVerbosity::Type LogVerbosityOverride = ELogVerbosity::NoLogging) = 0;
 
-	EInstallBundleResult FlushCache(FInstallBundleManagerFlushCacheCompleteDelegate Callback, ELogVerbosity::Type LogVerbosityOverride = ELogVerbosity::NoLogging);
+	INSTALLBUNDLEMANAGER_API EInstallBundleResult FlushCache(FInstallBundleManagerFlushCacheCompleteDelegate Callback, ELogVerbosity::Type LogVerbosityOverride = ELogVerbosity::NoLogging);
 	virtual EInstallBundleResult FlushCache(FInstallBundleSourceOrCache SourceOrCache, FInstallBundleManagerFlushCacheCompleteDelegate Callback, ELogVerbosity::Type LogVerbosityOverride = ELogVerbosity::NoLogging) = 0;
 
 	virtual TArray<FInstallBundleCacheStats> GetCacheStats(EInstallBundleCacheDumpToLog DumpToLog = EInstallBundleCacheDumpToLog::None, ELogVerbosity::Type LogVerbosityOverride = ELogVerbosity::NoLogging) = 0;
 	virtual TOptional<FInstallBundleCacheStats> GetCacheStats(FInstallBundleSourceOrCache SourceOrCache, EInstallBundleCacheDumpToLog DumpToLog = EInstallBundleCacheDumpToLog::None, ELogVerbosity::Type LogVerbosityOverride = ELogVerbosity::NoLogging) = 0;
 
-	void RequestRemoveContentOnNextInit(FName RemoveName, TArrayView<const FName> KeepNames = TArrayView<const FName>());
+	INSTALLBUNDLEMANAGER_API void RequestRemoveContentOnNextInit(FName RemoveName, TArrayView<const FName> KeepNames = TArrayView<const FName>());
 	virtual void RequestRemoveContentOnNextInit(TArrayView<const FName> RemoveNames, TArrayView<const FName> KeepNames = TArrayView<const FName>()) = 0;
 
-	void CancelRequestRemoveContentOnNextInit(FName BundleName);
+	INSTALLBUNDLEMANAGER_API void CancelRequestRemoveContentOnNextInit(FName BundleName);
 	virtual void CancelRequestRemoveContentOnNextInit(TArrayView<const FName> BundleNames) = 0;
 
-	void CancelUpdateContent(FName BundleName);
+	INSTALLBUNDLEMANAGER_API void CancelUpdateContent(FName BundleName);
 	virtual void CancelUpdateContent(TArrayView<const FName> BundleNames) = 0;
 
-	void PauseUpdateContent(FName BundleName);
+	INSTALLBUNDLEMANAGER_API void PauseUpdateContent(FName BundleName);
 	virtual void PauseUpdateContent(TArrayView<const FName> BundleNames) = 0;
 
-	void ResumeUpdateContent(FName BundleName);
+	INSTALLBUNDLEMANAGER_API void ResumeUpdateContent(FName BundleName);
 	virtual void ResumeUpdateContent(TArrayView<const FName> BundleNames) = 0;
 
 	virtual void RequestPausedBundleCallback() = 0;
@@ -165,10 +165,10 @@ public:
 	virtual TOptional<FInstallBundleProgress> GetBundleProgress(FName BundleName) const = 0;
 
 	virtual EInstallBundleRequestFlags GetModifyableContentRequestFlags() const = 0;
-	void UpdateContentRequestFlags(FName BundleName, EInstallBundleRequestFlags AddFlags, EInstallBundleRequestFlags RemoveFlags);
+	INSTALLBUNDLEMANAGER_API void UpdateContentRequestFlags(FName BundleName, EInstallBundleRequestFlags AddFlags, EInstallBundleRequestFlags RemoveFlags);
 	virtual void UpdateContentRequestFlags(TArrayView<const FName> BundleNames, EInstallBundleRequestFlags AddFlags, EInstallBundleRequestFlags RemoveFlags) = 0;
 	
-	virtual void StartPatchCheck();
+	INSTALLBUNDLEMANAGER_API virtual void StartPatchCheck();
 	virtual void AddEnvironmentWantsPatchCheckBackCompatDelegate(FName Tag, FInstallBundleManagerEnvironmentWantsPatchCheck Delegate) {}
 	virtual void RemoveEnvironmentWantsPatchCheckBackCompatDelegate(FName Tag) {}
 	virtual bool SupportsEarlyStartupPatching() const = 0;
diff --git a/Engine/Source/Runtime/InstallBundleManager/Public/InstallBundleTypes.h b/Engine/Source/Runtime/InstallBundleManager/Public/InstallBundleTypes.h
--- a/Engine/Source/Runtime/InstallBundleManager/Public/InstallBundleTypes.h
+++ b/Engine/Source/Runtime/InstallBundleManager/Public/InstallBundleTypes.h
@@ -72,22 +72,22 @@ enum class EInstallBundleInstallState : int
 };
 INSTALLBUNDLEMANAGER_API const TCHAR* LexToString(EInstallBundleInstallState State);
 
-struct INSTALLBUNDLEMANAGER_API FInstallBundleCombinedInstallState
+struct FInstallBundleCombinedInstallState
 {
 	TMap<FName, EInstallBundleInstallState> IndividualBundleStates;
 
-	bool GetAllBundlesHaveState(EInstallBundleInstallState State, TArrayView<const FName> ExcludedBundles = TArrayView<const FName>()) const;
-	bool GetAnyBundleHasState(EInstallBundleInstallState State, TArrayView<const FName> ExcludedBundles = TArrayView<const FName>()) const;
+	INSTALLBUNDLEMANAGER_API bool GetAllBundlesHaveState(EInstallBundleInstallState State, TArrayView<const FName> ExcludedBundles = TArrayView<const FName>()) const;
+	INSTALLBUNDLEMANAGER_API bool GetAnyBundleHasState(EInstallBundleInstallState State, TArrayView<const FName> ExcludedBundles = TArrayView<const FName>()) const;
 };
 
-struct INSTALLBUNDLEMANAGER_API FInstallBundleContentState
+struct FInstallBundleContentState
 {
 	EInstallBundleInstallState State = EInstallBundleInstallState::NotInstalled;
 	float Weight = 0.0f;
 	TMap<EInstallBundleSourceType, FString> Version;
 };
 
-struct INSTALLBUNDLEMANAGER_API FInstallBundleCombinedContentState
+struct FInstallBundleCombinedContentState
 {
 	TMap<FName, FInstallBundleContentState> IndividualBundleStates;
 	TMap<EInstallBundleSourceType, FString> CurrentVersion;
@@ -96,8 +96,8 @@ struct INSTALLBUNDLEMANAGER_API FInstallBundleCombinedContentState
 	uint64 InstallOverheadSize = 0;
 	uint64 FreeSpace = 0;
 
-	bool GetAllBundlesHaveState(EInstallBundleInstallState State, TArrayView<const FName> ExcludedBundles = TArrayView<const FName>()) const;	
-	bool GetAnyBundleHasState(EInstallBundleInstallState State, TArrayView<const FName> ExcludedBundles = TArrayView<const FName>()) const;
+	INSTALLBUNDLEMANAGER_API bool GetAllBundlesHaveState(EInstallBundleInstallState State, TArrayView<const FName> ExcludedBundles = TArrayView<const FName>()) const;	
+	INSTALLBUNDLEMANAGER_API bool GetAnyBundleHasState(EInstallBundleInstallState State, TArrayView<const FName> ExcludedBundles = TArrayView<const FName>()) const;
 };
 
 enum class EInstallBundleGetContentStateFlags : uint32
diff --git a/Engine/Source/Runtime/InstallBundleManager/Public/InstallBundleUtils.h b/Engine/Source/Runtime/InstallBundleManager/Public/InstallBundleUtils.h
--- a/Engine/Source/Runtime/InstallBundleManager/Public/InstallBundleUtils.h
+++ b/Engine/Source/Runtime/InstallBundleManager/Public/InstallBundleUtils.h
@@ -100,22 +100,22 @@ namespace InstallBundleUtil
 	}
 
 	// Keep the engine awake via RAII when running as an embedded app
-	class INSTALLBUNDLEMANAGER_API FInstallBundleManagerKeepAwake : public FEmbeddedKeepAwake
+	class FInstallBundleManagerKeepAwake : public FEmbeddedKeepAwake
 	{
-		static FName Tag;
-		static FName TagWithRendering;
+		static INSTALLBUNDLEMANAGER_API FName Tag;
+		static INSTALLBUNDLEMANAGER_API FName TagWithRendering;
 	public:
 		FInstallBundleManagerKeepAwake(bool bNeedsRendering = false)
 			: FEmbeddedKeepAwake(bNeedsRendering ? TagWithRendering : Tag, bNeedsRendering) {}
 	};
 
-	class INSTALLBUNDLEMANAGER_API FInstallBundleManagerScreenSaverControl
+	class FInstallBundleManagerScreenSaverControl
 	{
-		static bool bDidDisableScreensaver;
-		static int DisableCount;
+		static INSTALLBUNDLEMANAGER_API bool bDidDisableScreensaver;
+		static INSTALLBUNDLEMANAGER_API int DisableCount;
 
-		static void IncDisable();
-		static void DecDisable();
+		static INSTALLBUNDLEMANAGER_API void IncDisable();
+		static INSTALLBUNDLEMANAGER_API void DecDisable();
 
 	public:
 		FInstallBundleManagerScreenSaverControl()
@@ -143,21 +143,21 @@ namespace InstallBundleUtil
 	};
 
 	// Helper to safely suppress installer analytics that will restore previous state when going out of scope
-	struct INSTALLBUNDLEMANAGER_API FInstallBundleSuppressAnalytics : FNoncopyable
+	struct FInstallBundleSuppressAnalytics : FNoncopyable
 	{
-		FInstallBundleSuppressAnalytics();
-		~FInstallBundleSuppressAnalytics();
+		INSTALLBUNDLEMANAGER_API FInstallBundleSuppressAnalytics();
+		INSTALLBUNDLEMANAGER_API ~FInstallBundleSuppressAnalytics();
 
-		void Enable();
-		void Disable();
+		INSTALLBUNDLEMANAGER_API void Enable();
+		INSTALLBUNDLEMANAGER_API void Disable();
 
-		static bool IsEnabled();
+		static INSTALLBUNDLEMANAGER_API bool IsEnabled();
 
 	private:
 		bool bIsEnabled;
 	};
 
-	class INSTALLBUNDLEMANAGER_API FInstallBundleWork : public FNonAbandonableTask
+	class FInstallBundleWork : public FNonAbandonableTask
 	{
 	public:
 		FInstallBundleWork() = default;
@@ -203,7 +203,7 @@ namespace InstallBundleUtil
 
 	INSTALLBUNDLEMANAGER_API void CleanupInstallBundleAsyncIOTasks(TArray<TUniquePtr<FInstallBundleTask>>& Tasks);
 
-	struct INSTALLBUNDLEMANAGER_API FContentRequestStateStats
+	struct FContentRequestStateStats
 	{
 		double StartTime = 0.0;
 		double EndTime = 0.0;
@@ -216,7 +216,7 @@ namespace InstallBundleUtil
 		}
 	};
 
-	struct INSTALLBUNDLEMANAGER_API FContentRequestStats
+	struct FContentRequestStats
 	{
 		double StartTime = 0.0;
 		double EndTime = 0.0;
@@ -229,17 +229,17 @@ namespace InstallBundleUtil
 		}
 	};
 
-	class INSTALLBUNDLEMANAGER_API FContentRequestStatsMap
+	class FContentRequestStatsMap
 	{
 	private:
 		TMap<FName, InstallBundleUtil::FContentRequestStats> StatsMap;
 
 	public:
-		void StatsBegin(FName BundleName);
-		void StatsEnd(FName BundleName);
-		void StatsBegin(FName BundleName, const TCHAR* State);
-		void StatsEnd(FName BundleName, const TCHAR* State, uint64 DataSize = 0);
-		void StatsReset(FName BundleName);
+		INSTALLBUNDLEMANAGER_API void StatsBegin(FName BundleName);
+		INSTALLBUNDLEMANAGER_API void StatsEnd(FName BundleName);
+		INSTALLBUNDLEMANAGER_API void StatsBegin(FName BundleName, const TCHAR* State);
+		INSTALLBUNDLEMANAGER_API void StatsEnd(FName BundleName, const TCHAR* State, uint64 DataSize = 0);
+		INSTALLBUNDLEMANAGER_API void StatsReset(FName BundleName);
 
 		const TMap<FName, InstallBundleUtil::FContentRequestStats>& GetMap() { return StatsMap; }
 	};
@@ -314,11 +314,11 @@ namespace InstallBundleUtil
 
 		//Both ContentRequests and Individual Bundles use this to track data respective to themselves.
 		//Keeps us from having to write serialization / loading data for both cases that would be the same
-		class INSTALLBUNDLEMANAGER_API FPersistentStatsBase : public FJsonSerializable
+		class FPersistentStatsBase : public FJsonSerializable
 		{
 		public:
 			//We should not ever really be serializing this base class. Kept it here for reference.
-			BEGIN_JSON_SERIALIZER
+			INSTALLBUNDLEMANAGER_API BEGIN_JSON_SERIALIZER
 				JSON_SERIALIZE("AnalyticsSessionID", AnalyticsSessionID);
 				JSON_SERIALIZE_MAP("CountStats", CountStatMap);
 				JSON_SERIALIZE_MAP_SERIALIZABLE("TimingStats", TimingStatsMap, FPersistentTimerData);
@@ -328,50 +328,50 @@ namespace InstallBundleUtil
 			virtual ~FPersistentStatsBase() {}
 
 			//Tries to load persistent stats from disk and determine if we need to resume our previous persistent session or start a new one
-			void StatsBegin(const FString& ExpectedAnalyticsID, bool bForceResetData = false);
+			INSTALLBUNDLEMANAGER_API void StatsBegin(const FString& ExpectedAnalyticsID, bool bForceResetData = false);
 
 			//Tries to set this stat pool as "inactive". By default also stop all active timers. 
 			// NOTE: If you don't stop all active timers future calls to UpdateTimingStat will still increment this value and save it off!
-			void StatsEnd(bool bStopAllActiveTimers = true);
+			INSTALLBUNDLEMANAGER_API void StatsEnd(bool bStopAllActiveTimers = true);
 
-			bool IsTimingStatStarted(ETimingStatNames StatToUpdate) const;
-			void StartTimingStat(ETimingStatNames StatToUpdate);
-			void StopTimingStat(ETimingStatNames StatToUpdate, bool UpdateTimerOnStop = true);
-			void UpdateTimingStat(ETimingStatNames StatToUpdate);
+			INSTALLBUNDLEMANAGER_API bool IsTimingStatStarted(ETimingStatNames StatToUpdate) const;
+			INSTALLBUNDLEMANAGER_API void StartTimingStat(ETimingStatNames StatToUpdate);
+			INSTALLBUNDLEMANAGER_API void StopTimingStat(ETimingStatNames StatToUpdate, bool UpdateTimerOnStop = true);
+			INSTALLBUNDLEMANAGER_API void UpdateTimingStat(ETimingStatNames StatToUpdate);
 			
-			void UpdateAllActiveTimers();
-			void StopAllActiveTimers();
+			INSTALLBUNDLEMANAGER_API void UpdateAllActiveTimers();
+			INSTALLBUNDLEMANAGER_API void StopAllActiveTimers();
 			
 			bool IsActive() const { return bIsActive; }
 			bool IsDirty() const { return bIsDirty; }
 
-			void IncrementCountStat(PersistentStats::ECountStatNames StatToUpdate);
+			INSTALLBUNDLEMANAGER_API void IncrementCountStat(PersistentStats::ECountStatNames StatToUpdate);
 
-			bool HasTimingStat(ETimingStatNames StatToCheck) const;
-			bool HasCountStat(ECountStatNames StatToCheck) const;
+			INSTALLBUNDLEMANAGER_API bool HasTimingStat(ETimingStatNames StatToCheck) const;
+			INSTALLBUNDLEMANAGER_API bool HasCountStat(ECountStatNames StatToCheck) const;
 
-			const FPersistentTimerData* GetTimingStatData(ETimingStatNames StatToGet) const;
-			const int* GetCountStatData(ECountStatNames StatToGet) const;
+			INSTALLBUNDLEMANAGER_API const FPersistentTimerData* GetTimingStatData(ETimingStatNames StatToGet) const;
+			INSTALLBUNDLEMANAGER_API const int* GetCountStatData(ECountStatNames StatToGet) const;
 
 			//Saves the persistent stat data to disk in the location returned by GetFullPathForStatFile()
 			//Returns True if that save succeeds and false otherwise.
-			bool SaveStatsToDisk();
+			INSTALLBUNDLEMANAGER_API bool SaveStatsToDisk();
 			
 			//Function that allows you to load stats from disk. Returns true if stats were either loaded from disk now, or previously were loaded from disk.
 			//NOTE: This does not begin stat collection, so you should still call StatsBegin on this before handling any stats!
-			bool LoadStatsFromDisk();
+			INSTALLBUNDLEMANAGER_API bool LoadStatsFromDisk();
 
 		protected:
-			void ResetStats(const FString& NewAnalyticsSessionID);
+			INSTALLBUNDLEMANAGER_API void ResetStats(const FString& NewAnalyticsSessionID);
 			
 			//Called after we load data but before we process any of it.
 			//Useful for fixing up the data loaded from disk before it changes or is acted upon by anything else
 			//NOTE: A call to LoadStatsFromDisk will not necessarily call this! It will only be called when data is actually pulled from
 			//the disk, and its possible to skip pulling data from the disk!
-			void OnLoadingDataFromDisk();
+			INSTALLBUNDLEMANAGER_API void OnLoadingDataFromDisk();
 
 			//Helper to try and reconcile offline and active timers after we load data from disk
-			void HandleTimerStatsAfterDataLoad();
+			INSTALLBUNDLEMANAGER_API void HandleTimerStatsAfterDataLoad();
 
 			//Protected so we can only construct the derived versions of this Base class
 			FPersistentStatsBase()
@@ -399,7 +399,7 @@ namespace InstallBundleUtil
 		//Static Methods
 		public:
 			//Helper function that generates a basic ExpectedAnalyticsID from device and UE information
-			static const FString GetBaseExpectedAnalyticsID();
+			static INSTALLBUNDLEMANAGER_API const FString GetBaseExpectedAnalyticsID();
 		
 		//Pure Virtual Methods
 		public:
@@ -408,10 +408,10 @@ namespace InstallBundleUtil
 		};
 
 		//Each bundles persistent stats
-		class INSTALLBUNDLEMANAGER_API FBundlePersistentStats : public FPersistentStatsBase
+		class FBundlePersistentStats : public FPersistentStatsBase
 		{
 		public:
-			BEGIN_JSON_SERIALIZER
+			INSTALLBUNDLEMANAGER_API BEGIN_JSON_SERIALIZER
 				JSON_SERIALIZE("BundleName", BundleName);
 				JSON_SERIALIZE("AnalyticsSessionID", AnalyticsSessionID);
 				JSON_SERIALIZE_MAP("CountStats", CountStatMap);
@@ -436,14 +436,14 @@ namespace InstallBundleUtil
 
 			//Implementation for FInstallBundleStatsContainerBase
 		public:
-			virtual const FString GetFullPathForStatFile() const override;
+			INSTALLBUNDLEMANAGER_API virtual const FString GetFullPathForStatFile() const override;
 		};
 
 		//Tracks a set of bundles and stats for a named session. This way you can have stats across multiple bundle installs grouped by some naming convention
-		class INSTALLBUNDLEMANAGER_API FSessionPersistentStats : public FPersistentStatsBase
+		class FSessionPersistentStats : public FPersistentStatsBase
 		{
 		public:
-			BEGIN_JSON_SERIALIZER
+			INSTALLBUNDLEMANAGER_API BEGIN_JSON_SERIALIZER
 				JSON_SERIALIZE("SessionName", SessionName);
 				JSON_SERIALIZE_ARRAY("RequiredBundles", RequiredBundles);
 				JSON_SERIALIZE("AnalyticsSessionID", AnalyticsSessionID);
@@ -470,103 +470,103 @@ namespace InstallBundleUtil
 
 			virtual ~FSessionPersistentStats() {}
 
-			void AddRequiredBundles(const TArray<FString>& RequiredBundlesToAdd);
-			void AddRequiredBundles(const TArray<FName>& RequiredBundlesToAdd);
+			INSTALLBUNDLEMANAGER_API void AddRequiredBundles(const TArray<FString>& RequiredBundlesToAdd);
+			INSTALLBUNDLEMANAGER_API void AddRequiredBundles(const TArray<FName>& RequiredBundlesToAdd);
 
 			//Removes all Required Bundle Names except the ones in the NewRequiredBundlesList. If no NewRequiredBundlesList is supplied, just erases all RequiredBundles.
 			//This includes ones loaded from the persistent cache from previous runs. This re-saves the persistent cache file so they will be removed completely.
-			void ResetRequiredBundles(const TArray<FString>& NewRequiredBundles = TArray<FString>());
+			INSTALLBUNDLEMANAGER_API void ResetRequiredBundles(const TArray<FString>& NewRequiredBundles = TArray<FString>());
 
 			//Populates the passed in TArray with all the required bundles' names
-			void GetRequiredBundles(TArray<FString>& OutRequiredBundles) const;
+			INSTALLBUNDLEMANAGER_API void GetRequiredBundles(TArray<FString>& OutRequiredBundles) const;
 			
 			//Implementation for FInstallBundleStatsContainerBase
 		public:
-			virtual const FString GetFullPathForStatFile() const override;
+			INSTALLBUNDLEMANAGER_API virtual const FString GetFullPathForStatFile() const override;
 		};
 
 		//Helper class to handle holding a collection of Bundle and Session stats.
-		class INSTALLBUNDLEMANAGER_API FPersistentStatContainerBase
+		class FPersistentStatContainerBase
 		{
 		public:
-			FPersistentStatContainerBase();
-			virtual ~FPersistentStatContainerBase();
+			INSTALLBUNDLEMANAGER_API FPersistentStatContainerBase();
+			INSTALLBUNDLEMANAGER_API virtual ~FPersistentStatContainerBase();
 
 			//NOTE! These 2 functions are not virtual as we call them in the construct/deconstructor!
-			void InitializeBase();
-			void ShutdownBase();
+			INSTALLBUNDLEMANAGER_API void InitializeBase();
+			INSTALLBUNDLEMANAGER_API void ShutdownBase();
 
 			//Starts Bundle Persistent Stat Tracking for the given BundleName. Will automatically load applicable data from the disk and reset that cache
 			//if the ExpectedAnalyticsID doesn't match the one previously serialized.
 			//Uses FInstallBundleStatsContainerBase::GetBaseExpectedAnalyticsID if an ExpectedAnalyticsID is not passed in or empty
-			virtual void StartBundlePersistentStatTracking(FName BundleName, const FString& ExpectedAnalyticsID = FString(), bool bForceResetStatData = false);
+			INSTALLBUNDLEMANAGER_API virtual void StartBundlePersistentStatTracking(FName BundleName, const FString& ExpectedAnalyticsID = FString(), bool bForceResetStatData = false);
 			
 			//Starts Session Persistent Stat Tracking under the given SessionName. Will automatically load applicable data from the disk and reset that cache
 			//if the ExpectedAnalyticsID doesn't match the one previously serialized. 
 			//Uses FInstallBundleStatsContainerBase::GetBaseExpectedAnalyticsID if an ExpectedAnalyticsID is not passed in or empty
-			virtual void StartSessionPersistentStatTracking(const FString& SessionName, const TArray<FName>& RequiredBundles = TArray<FName>(), const FString& ExpectedAnalyticsID = FString(), bool bForceResetStatData = false);
+			INSTALLBUNDLEMANAGER_API virtual void StartSessionPersistentStatTracking(const FString& SessionName, const TArray<FName>& RequiredBundles = TArray<FName>(), const FString& ExpectedAnalyticsID = FString(), bool bForceResetStatData = false);
 
-			virtual void StopBundlePersistentStatTracking(FName BundleName, bool bStopAllActiveTimers = true);
-			virtual void StopSessionPersistentStatTracking(const FString& SessionName, bool bStopAllActiveTimers = true);
+			INSTALLBUNDLEMANAGER_API virtual void StopBundlePersistentStatTracking(FName BundleName, bool bStopAllActiveTimers = true);
+			INSTALLBUNDLEMANAGER_API virtual void StopSessionPersistentStatTracking(const FString& SessionName, bool bStopAllActiveTimers = true);
 
-			virtual void StartBundlePersistentStatTimer(FName BundleName, ETimingStatNames TimerToStart);
-			virtual void StartSessionPersistentStatTimer(const FString& SessionName, ETimingStatNames TimerToStart);
+			INSTALLBUNDLEMANAGER_API virtual void StartBundlePersistentStatTimer(FName BundleName, ETimingStatNames TimerToStart);
+			INSTALLBUNDLEMANAGER_API virtual void StartSessionPersistentStatTimer(const FString& SessionName, ETimingStatNames TimerToStart);
 
-			virtual void StopBundlePersistentStatTimer(FName BundleName, ETimingStatNames TimerToStop);
-			virtual void StopSessionPersistentStatTimer(const FString& SessionName, ETimingStatNames TimerToStop);
+			INSTALLBUNDLEMANAGER_API virtual void StopBundlePersistentStatTimer(FName BundleName, ETimingStatNames TimerToStop);
+			INSTALLBUNDLEMANAGER_API virtual void StopSessionPersistentStatTimer(const FString& SessionName, ETimingStatNames TimerToStop);
 
-			virtual void UpdateBundlePersistentStatTimer(FName BundleName, ETimingStatNames TimerToUpdate);
-			virtual void UpdateSessionPersistentStatTimer(const FString& SessionName, ETimingStatNames TimerToUpdate);
+			INSTALLBUNDLEMANAGER_API virtual void UpdateBundlePersistentStatTimer(FName BundleName, ETimingStatNames TimerToUpdate);
+			INSTALLBUNDLEMANAGER_API virtual void UpdateSessionPersistentStatTimer(const FString& SessionName, ETimingStatNames TimerToUpdate);
 
-			virtual void IncrementBundlePersistentCounter(FName BundleName, ECountStatNames CounterToUpdate);
-			virtual void IncrementSessionPersistentCounter(const FString& SessionName, ECountStatNames CounterToUpdate);
+			INSTALLBUNDLEMANAGER_API virtual void IncrementBundlePersistentCounter(FName BundleName, ECountStatNames CounterToUpdate);
+			INSTALLBUNDLEMANAGER_API virtual void IncrementSessionPersistentCounter(const FString& SessionName, ECountStatNames CounterToUpdate);
 
-			virtual const FBundlePersistentStats* GetBundleStat(FName BundleName) const;
-			virtual const FSessionPersistentStats* GetSessionStat(const FString& SessionName) const;
+			INSTALLBUNDLEMANAGER_API virtual const FBundlePersistentStats* GetBundleStat(FName BundleName) const;
+			INSTALLBUNDLEMANAGER_API virtual const FSessionPersistentStats* GetSessionStat(const FString& SessionName) const;
 
-			virtual void SaveAllDirtyStatsToDisk();
+			INSTALLBUNDLEMANAGER_API virtual void SaveAllDirtyStatsToDisk();
 			
 			//Deletes persistent stat information from our stats for this Session/Bundle to reduce memory useage
-			virtual void RemoveSessionStats(const FString& SessionName);
-			virtual void RemoveBundleStats(FName BundleName);
+			INSTALLBUNDLEMANAGER_API virtual void RemoveSessionStats(const FString& SessionName);
+			INSTALLBUNDLEMANAGER_API virtual void RemoveBundleStats(FName BundleName);
 			
 		protected:
-			virtual bool Tick(float dt);
-			virtual void ResetTimerUpdate();
-			virtual void ResetDirtyStatUpdate();
+			INSTALLBUNDLEMANAGER_API virtual bool Tick(float dt);
+			INSTALLBUNDLEMANAGER_API virtual void ResetTimerUpdate();
+			INSTALLBUNDLEMANAGER_API virtual void ResetDirtyStatUpdate();
 
-			virtual void UpdateAllBundlesActiveTimers();
-			virtual void UpdateAllSessionActiveTimers();
+			INSTALLBUNDLEMANAGER_API virtual void UpdateAllBundlesActiveTimers();
+			INSTALLBUNDLEMANAGER_API virtual void UpdateAllSessionActiveTimers();
 
-			virtual void OnApp_EnteringBackground();
-			virtual void OnApp_EnteringForeground();
+			INSTALLBUNDLEMANAGER_API virtual void OnApp_EnteringBackground();
+			INSTALLBUNDLEMANAGER_API virtual void OnApp_EnteringForeground();
 
-			virtual void OnBackground_HandleBundleStats();
-			virtual void OnForeground_HandleBundleStats();
-			virtual void OnBackground_HandleSessionStats();
-			virtual void OnForeground_HandleSessionStats();
+			INSTALLBUNDLEMANAGER_API virtual void OnBackground_HandleBundleStats();
+			INSTALLBUNDLEMANAGER_API virtual void OnForeground_HandleBundleStats();
+			INSTALLBUNDLEMANAGER_API virtual void OnBackground_HandleSessionStats();
+			INSTALLBUNDLEMANAGER_API virtual void OnForeground_HandleSessionStats();
 			
 			//Called whenever we start a timer
 			//NOTE: doesn't get called by timer's being Start/Stopped for bShouldAutoHandleFGBGStats
-			virtual void OnTimerStartedForStat(FPersistentStatsBase& BundleStatForTimer, ETimingStatNames TimerStarted);
+			INSTALLBUNDLEMANAGER_API virtual void OnTimerStartedForStat(FPersistentStatsBase& BundleStatForTimer, ETimingStatNames TimerStarted);
 
 			//Called whenever we stop a timer 
 			//NOTE: doesn't get called by timer's being Start/Stopped for bShouldAutoHandleFGBGStats
-			virtual void OnTimerStoppedForStat(FPersistentStatsBase& BundleStatForTimer, ETimingStatNames TimerStarted);
+			INSTALLBUNDLEMANAGER_API virtual void OnTimerStoppedForStat(FPersistentStatsBase& BundleStatForTimer, ETimingStatNames TimerStarted);
 
 			//Stops active foreground timers when going to background and starts applicable background version
 			//Also increments background counters
-			virtual void UpdateStatsForBackground(FPersistentStatsBase& StatToUpdate);
+			INSTALLBUNDLEMANAGER_API virtual void UpdateStatsForBackground(FPersistentStatsBase& StatToUpdate);
 
 			//Stops active background timers and resumes applicable foreground timers
 			//Also increments foreground counters
-			virtual void UpdateStatsForForeground(FPersistentStatsBase& StatToUpdate);
+			INSTALLBUNDLEMANAGER_API virtual void UpdateStatsForForeground(FPersistentStatsBase& StatToUpdate);
 			
 			//Goes through this Session's RequiredBundles and loads data from disk for those bundles. Allows us to make sure we have
 			//full data in memory from disk without having to know the previous runs AnalyticsID to call StatsBegin on that bundle.
 			//NOTE: You should still call StatsBegin on the listed bundles before Starting/Incrementing/ETC any analytics! This just
 			//ensures that the previous run data exists for this given session!
-			virtual void LoadRequiredBundleDataFromDiskForSession(const FString& SessionName);
+			INSTALLBUNDLEMANAGER_API virtual void LoadRequiredBundleDataFromDiskForSession(const FString& SessionName);
 			
 		protected:
 			TMap<FName, FBundlePersistentStats> PerBundlePersistentStatMap;
@@ -618,12 +618,12 @@ namespace InstallBundleUtil
 		};
 	}
 
-	struct INSTALLBUNDLEMANAGER_API IBundleSourceContentRequestSharedContext
+	struct IBundleSourceContentRequestSharedContext
 	{
 		virtual ~IBundleSourceContentRequestSharedContext() {}
 	};
 
-	struct INSTALLBUNDLEMANAGER_API FContentRequestSharedContext
+	struct FContentRequestSharedContext
 	{
 		FContentRequestSharedContext() = default;
 		FContentRequestSharedContext(const FContentRequestSharedContext& Other) = delete;
diff --git a/Engine/Source/Runtime/InteractiveToolsFramework/Public/BaseBehaviors/AnyButtonInputBehavior.h b/Engine/Source/Runtime/InteractiveToolsFramework/Public/BaseBehaviors/AnyButtonInputBehavior.h
--- a/Engine/Source/Runtime/InteractiveToolsFramework/Public/BaseBehaviors/AnyButtonInputBehavior.h
+++ b/Engine/Source/Runtime/InteractiveToolsFramework/Public/BaseBehaviors/AnyButtonInputBehavior.h
@@ -28,48 +28,48 @@ class UObject;
  *  @todo support tablet fingers
  *  @todo support gamepad?
  */
-UCLASS()
-class INTERACTIVETOOLSFRAMEWORK_API UAnyButtonInputBehavior : public UInputBehavior
+UCLASS(MinimalAPI)
+class UAnyButtonInputBehavior : public UInputBehavior
 {
 	GENERATED_BODY()
 
 public:
-	UAnyButtonInputBehavior();
+	INTERACTIVETOOLSFRAMEWORK_API UAnyButtonInputBehavior();
 
 	/** Return set of devices supported by this behavior */
-	virtual EInputDevices GetSupportedDevices() override;
+	INTERACTIVETOOLSFRAMEWORK_API virtual EInputDevices GetSupportedDevices() override;
 
 	/** @return true if Target Button has been pressed this frame */
-	virtual bool IsPressed(const FInputDeviceState& input);
+	INTERACTIVETOOLSFRAMEWORK_API virtual bool IsPressed(const FInputDeviceState& input);
 
 	/** @return true if Target Button is currently held down */
-	virtual bool IsDown(const FInputDeviceState& input);
+	INTERACTIVETOOLSFRAMEWORK_API virtual bool IsDown(const FInputDeviceState& input);
 
 	/** @return true if Target Button was released this frame */
-	virtual bool IsReleased(const FInputDeviceState& input);
+	INTERACTIVETOOLSFRAMEWORK_API virtual bool IsReleased(const FInputDeviceState& input);
 	
 	/** @return current 2D position of Target Device, or zero if device does not have 2D position */
-	virtual FVector2D GetClickPoint(const FInputDeviceState& input);
+	INTERACTIVETOOLSFRAMEWORK_API virtual FVector2D GetClickPoint(const FInputDeviceState& input);
 
 	/** @return current 3D world ray for Target Device position */
-	virtual FRay GetWorldRay(const FInputDeviceState& input);
+	INTERACTIVETOOLSFRAMEWORK_API virtual FRay GetWorldRay(const FInputDeviceState& input);
 	
 	/** @return current 3D world ray and optional 2D position for Target Device */
-	virtual FInputDeviceRay GetDeviceRay(const FInputDeviceState& input);
+	INTERACTIVETOOLSFRAMEWORK_API virtual FInputDeviceRay GetDeviceRay(const FInputDeviceState& input);
 
 	/** @return last-active supported Device */
-	EInputDevices GetActiveDevice() const;
+	INTERACTIVETOOLSFRAMEWORK_API EInputDevices GetActiveDevice() const;
 
 
 public:
 	/** Configure the target Mouse button to be the left button */
-	virtual void SetUseLeftMouseButton();
+	INTERACTIVETOOLSFRAMEWORK_API virtual void SetUseLeftMouseButton();
 	/** Configure the target Mouse button to be the middle button */
-	virtual void SetUseMiddleMouseButton();
+	INTERACTIVETOOLSFRAMEWORK_API virtual void SetUseMiddleMouseButton();
 	/** Configure the target Mouse button to be the right button */
-	virtual void SetUseRightMouseButton();
+	INTERACTIVETOOLSFRAMEWORK_API virtual void SetUseRightMouseButton();
 	/** Configure a custom target mouse button */
-	virtual void SetUseCustomMouseButton(TUniqueFunction<FDeviceButtonState(const FInputDeviceState& Input)>);
+	INTERACTIVETOOLSFRAMEWORK_API virtual void SetUseCustomMouseButton(TUniqueFunction<FDeviceButtonState(const FInputDeviceState& Input)>);
 
 
 protected:
@@ -80,5 +80,5 @@ protected:
 	TUniqueFunction<FDeviceButtonState(const FInputDeviceState& Input)> GetMouseButtonStateFunc;
 
 	/** Returns FDeviceButtonState for target active device button */
-	FDeviceButtonState GetButtonStateFunc(const FInputDeviceState& Input);
-};
\ No newline at end of file
+	INTERACTIVETOOLSFRAMEWORK_API FDeviceButtonState GetButtonStateFunc(const FInputDeviceState& Input);
+};
diff --git a/Engine/Source/Runtime/InteractiveToolsFramework/Public/BaseBehaviors/ClickDragBehavior.h b/Engine/Source/Runtime/InteractiveToolsFramework/Public/BaseBehaviors/ClickDragBehavior.h
--- a/Engine/Source/Runtime/InteractiveToolsFramework/Public/BaseBehaviors/ClickDragBehavior.h
+++ b/Engine/Source/Runtime/InteractiveToolsFramework/Public/BaseBehaviors/ClickDragBehavior.h
@@ -28,13 +28,13 @@ class UObject;
  *    
  * If a ForceEndCapture occurs we call Target::OnTerminateDragSequence   
  */
-UCLASS()
-class INTERACTIVETOOLSFRAMEWORK_API UClickDragInputBehavior : public UAnyButtonInputBehavior
+UCLASS(MinimalAPI)
+class UClickDragInputBehavior : public UAnyButtonInputBehavior
 {
 	GENERATED_BODY()
 
 public:
-	UClickDragInputBehavior();
+	INTERACTIVETOOLSFRAMEWORK_API UClickDragInputBehavior();
 
 	/**
 	 * The modifier set for this behavior
@@ -45,7 +45,7 @@ public:
 	 * Initialize this behavior with the given Target
 	 * @param Target implementor of hit-test and on-clicked functions
 	 */
-	virtual void Initialize(IClickDragBehaviorTarget* Target);
+	INTERACTIVETOOLSFRAMEWORK_API virtual void Initialize(IClickDragBehaviorTarget* Target);
 
 
 	/**
@@ -56,10 +56,10 @@ public:
 
 	// UInputBehavior implementation
 
-	virtual FInputCaptureRequest WantsCapture(const FInputDeviceState& Input) override;
-	virtual FInputCaptureUpdate BeginCapture(const FInputDeviceState& Input, EInputCaptureSide Side) override;
-	virtual FInputCaptureUpdate UpdateCapture(const FInputDeviceState& Input, const FInputCaptureData& Data) override;
-	virtual void ForceEndCapture(const FInputCaptureData& Data) override;
+	INTERACTIVETOOLSFRAMEWORK_API virtual FInputCaptureRequest WantsCapture(const FInputDeviceState& Input) override;
+	INTERACTIVETOOLSFRAMEWORK_API virtual FInputCaptureUpdate BeginCapture(const FInputDeviceState& Input, EInputCaptureSide Side) override;
+	INTERACTIVETOOLSFRAMEWORK_API virtual FInputCaptureUpdate UpdateCapture(const FInputDeviceState& Input, const FInputCaptureData& Data) override;
+	INTERACTIVETOOLSFRAMEWORK_API virtual void ForceEndCapture(const FInputCaptureData& Data) override;
 
 	/**
 	 * If true, then we will update Modifier states in UpdateCapture(). This defaults to false because
@@ -79,17 +79,17 @@ protected:
 	/**
 	 * Internal function that forwards click evens to Target::OnClickPress, you can customize behavior here
 	 */
-	virtual void OnClickPressInternal(const FInputDeviceState& Input, EInputCaptureSide Side);
+	INTERACTIVETOOLSFRAMEWORK_API virtual void OnClickPressInternal(const FInputDeviceState& Input, EInputCaptureSide Side);
 
 	/**
 	 * Internal function that forwards click evens to Target::OnClickDrag, you can customize behavior here
 	 */
-	virtual void OnClickDragInternal(const FInputDeviceState& Input, const FInputCaptureData& Data);
+	INTERACTIVETOOLSFRAMEWORK_API virtual void OnClickDragInternal(const FInputDeviceState& Input, const FInputCaptureData& Data);
 
 	/**
 	 * Internal function that forwards click evens to Target::OnClickRelease, you can customize behavior here
 	 */
-	virtual void OnClickReleaseInternal(const FInputDeviceState& Input, const FInputCaptureData& Data);
+	INTERACTIVETOOLSFRAMEWORK_API virtual void OnClickReleaseInternal(const FInputDeviceState& Input, const FInputCaptureData& Data);
 };
 
 
@@ -101,8 +101,8 @@ protected:
  * via a set of local lambda functions. To use/customize this class the client replaces the lambda functions with their own.
  * This avoids having to create a second IClickDragBehaviorTarget implementation for trivial use-cases.
  */
-UCLASS()
-class INTERACTIVETOOLSFRAMEWORK_API ULocalClickDragInputBehavior : public UClickDragInputBehavior, public IClickDragBehaviorTarget
+UCLASS(MinimalAPI)
+class ULocalClickDragInputBehavior : public UClickDragInputBehavior, public IClickDragBehaviorTarget
 {
 	GENERATED_BODY()
 protected:
diff --git a/Engine/Source/Runtime/InteractiveToolsFramework/Public/BaseBehaviors/KeyAsModifierInputBehavior.h b/Engine/Source/Runtime/InteractiveToolsFramework/Public/BaseBehaviors/KeyAsModifierInputBehavior.h
--- a/Engine/Source/Runtime/InteractiveToolsFramework/Public/BaseBehaviors/KeyAsModifierInputBehavior.h
+++ b/Engine/Source/Runtime/InteractiveToolsFramework/Public/BaseBehaviors/KeyAsModifierInputBehavior.h
@@ -26,13 +26,13 @@ class UObject;
  * This means that the modifier won't be updated if another behavior captures the
  * keyboard.
  */
-UCLASS()
-class INTERACTIVETOOLSFRAMEWORK_API UKeyAsModifierInputBehavior : public UInputBehavior
+UCLASS(MinimalAPI)
+class UKeyAsModifierInputBehavior : public UInputBehavior
 {
 	GENERATED_BODY()
 
 public:
-	UKeyAsModifierInputBehavior();
+	INTERACTIVETOOLSFRAMEWORK_API UKeyAsModifierInputBehavior();
 
 	virtual EInputDevices GetSupportedDevices() override
 	{
@@ -45,7 +45,7 @@ public:
 	 * @param ModifierID integer ID that identifiers the modifier toggle
 	 * @param ModifierKey the key that will be used as the modifier toggle
 	 */
-	virtual void Initialize(IModifierToggleBehaviorTarget* Target, int ModifierID, const FKey& ModifierKey);
+	INTERACTIVETOOLSFRAMEWORK_API virtual void Initialize(IModifierToggleBehaviorTarget* Target, int ModifierID, const FKey& ModifierKey);
 
 	/**
 	 * Initialize this behavior with the given Target
@@ -53,13 +53,13 @@ public:
 	 * @param ModifierID integer ID that identifiers the modifier toggle
 	 * @param ModifierCheckFunc the function that will be checked to set the the modifier.
 	 */
-	void Initialize(IModifierToggleBehaviorTarget* TargetIn, int ModifierID, TFunction<bool(const FInputDeviceState&)> ModifierCheckFunc);
+	INTERACTIVETOOLSFRAMEWORK_API void Initialize(IModifierToggleBehaviorTarget* TargetIn, int ModifierID, TFunction<bool(const FInputDeviceState&)> ModifierCheckFunc);
 
 	// UInputBehavior implementation
-	virtual FInputCaptureRequest WantsCapture(const FInputDeviceState& Input) override;
-	virtual FInputCaptureUpdate BeginCapture(const FInputDeviceState& Input, EInputCaptureSide Side) override;
-	virtual FInputCaptureUpdate UpdateCapture(const FInputDeviceState& Input, const FInputCaptureData& Data) override;
-	virtual void ForceEndCapture(const FInputCaptureData& Data) override;
+	INTERACTIVETOOLSFRAMEWORK_API virtual FInputCaptureRequest WantsCapture(const FInputDeviceState& Input) override;
+	INTERACTIVETOOLSFRAMEWORK_API virtual FInputCaptureUpdate BeginCapture(const FInputDeviceState& Input, EInputCaptureSide Side) override;
+	INTERACTIVETOOLSFRAMEWORK_API virtual FInputCaptureUpdate UpdateCapture(const FInputDeviceState& Input, const FInputCaptureData& Data) override;
+	INTERACTIVETOOLSFRAMEWORK_API virtual void ForceEndCapture(const FInputCaptureData& Data) override;
 
 
 protected:
