Add a toggle for the behavior of EKeys::Platform_Delete on Mac.

On a Mac keyboard pressing just "delete" will act as a "backspace" on windows or linux, and pressing "function+delete" will act as "delete" on windows or linux.

There is no behavioral change here, but it provides the option to change what it is and supports more platforms then just Mac.

diff --git a/Engine/Config/Mac/BaseMacInput.ini b/Engine/Config/Mac/BaseMacInput.ini
new file mode 100644
--- /dev/null
+++ b/Engine/Config/Mac/BaseMacInput.ini
@@ -0,0 +1,8 @@
+[MacInput]
+; This changes what EKeys::Platform_Delete will be set to on the Mac platform.
+;
+; If this is true (default), then EKeys::Platform_Delete will be EKeys::Backspace.
+; If this is false, then EKeys::Platform_Delete will be set to EKeys::Delete.
+; 
+; See FMacPlatformInput::GetPlatformDeleteKey for implementation
+bPlatformDeleteIsBackspace=true
\ No newline at end of file
diff --git a/Engine/Source/Runtime/InputCore/Private/InputCoreTypes.cpp b/Engine/Source/Runtime/InputCore/Private/InputCoreTypes.cpp
--- a/Engine/Source/Runtime/InputCore/Private/InputCoreTypes.cpp
+++ b/Engine/Source/Runtime/InputCore/Private/InputCoreTypes.cpp
@@ -159,13 +159,8 @@ const FKey EKeys::E_AccentAigu("E_AccentAigu");
 const FKey EKeys::C_Cedille("C_Cedille");
 const FKey EKeys::Section("Section");
 
-
 // Setup platform specific keys
-#if PLATFORM_MAC
-const FKey EKeys::Platform_Delete = EKeys::BackSpace;
-#else
-const FKey EKeys::Platform_Delete = EKeys::Delete;
-#endif
+const FKey EKeys::Platform_Delete = FPlatformInput::GetPlatformDeleteKey();
 
 // Ensure that the Gamepad_ names match those in GenericApplication.cpp
 const FKey EKeys::Gamepad_Left2D("Gamepad_Left2D");
diff --git a/Engine/Source/Runtime/InputCore/Private/Mac/MacPlatformInput.cpp b/Engine/Source/Runtime/InputCore/Private/Mac/MacPlatformInput.cpp
--- a/Engine/Source/Runtime/InputCore/Private/Mac/MacPlatformInput.cpp
+++ b/Engine/Source/Runtime/InputCore/Private/Mac/MacPlatformInput.cpp
@@ -1,6 +1,7 @@
 // Copyright Epic Games, Inc. All Rights Reserved.
 
 #include "Mac/MacPlatformInput.h"
+#include "Misc/ConfigCacheIni.h"	// For GConfig
 
 uint32 FMacPlatformInput::GetKeyMap( uint32* KeyCodes, FString* KeyNames, uint32 MaxMappings )
 {
@@ -81,3 +82,15 @@ uint32 FMacPlatformInput::GetCharKeyMap(uint32* KeyCodes, FString* KeyNames, uin
 	return FGenericPlatformInput::GetStandardPrintableKeyMap(KeyCodes, KeyNames, MaxMappings, false, true);
 }
 
+FKey FMacPlatformInput::GetPlatformDeleteKey()
+{
+	// The legacy behavior is that Platform Delete is "Backspace". In 5.4 we added this option to have it
+	// be the normal "Delete" key. 
+	bool bPlatformDeleteIsBackspace = true;
+	if (GConfig)
+	{
+		GConfig->GetBool(TEXT("MacInput"), TEXT("bPlatformDeleteIsBackspace"), bPlatformDeleteIsBackspace, GInputIni);
+	}
+
+	return bPlatformDeleteIsBackspace ? EKeys::BackSpace : EKeys::Delete;
+}
\ No newline at end of file
diff --git a/Engine/Source/Runtime/InputCore/Public/GenericPlatform/GenericPlatformInput.h b/Engine/Source/Runtime/InputCore/Public/GenericPlatform/GenericPlatformInput.h
--- a/Engine/Source/Runtime/InputCore/Public/GenericPlatform/GenericPlatformInput.h
+++ b/Engine/Source/Runtime/InputCore/Public/GenericPlatform/GenericPlatformInput.h
@@ -34,6 +34,11 @@ public:
 		return EKeys::Gamepad_FaceButton_Right;
 	}
 
+	static FKey GetPlatformDeleteKey()
+	{
+		return EKeys::Delete;
+	}
+
 protected:
 	/**
 	* Retrieves some standard key code mappings (usually called by a subclass's GetCharKeyMap)
diff --git a/Engine/Source/Runtime/InputCore/Public/Mac/MacPlatformInput.h b/Engine/Source/Runtime/InputCore/Public/Mac/MacPlatformInput.h
--- a/Engine/Source/Runtime/InputCore/Public/Mac/MacPlatformInput.h
+++ b/Engine/Source/Runtime/InputCore/Public/Mac/MacPlatformInput.h
@@ -9,6 +9,7 @@ struct INPUTCORE_API FMacPlatformInput : FGenericPlatformInput
 {
 	static uint32 GetKeyMap( uint32* KeyCodes, FString* KeyNames, uint32 MaxMappings );
 	static uint32 GetCharKeyMap(uint32* KeyCodes, FString* KeyNames, uint32 MaxMappings);
+	static FKey GetPlatformDeleteKey();
 };
 
 typedef FMacPlatformInput FPlatformInput;
