Add LLM memory tags for EKeys and the FKeyManager to reduce the number of "untagged" allocations

#rnx

diff --git a/Engine/Plugins/EnhancedInput/Source/InputEditor/Private/InputCustomizations.cpp b/Engine/Plugins/EnhancedInput/Source/InputEditor/Private/InputCustomizations.cpp
--- a/Engine/Plugins/EnhancedInput/Source/InputEditor/Private/InputCustomizations.cpp
+++ b/Engine/Plugins/EnhancedInput/Source/InputEditor/Private/InputCustomizations.cpp
@@ -255,6 +255,12 @@ bool FEnhancedActionMappingCustomization::DoesTriggerArrayContainCombo() const
 
 TSet<FName> FEnhancedInputDeveloperSettingsCustomization::ExcludedAssetNames;
 
+TSharedRef<IDetailCustomization> FEnhancedInputDeveloperSettingsCustomization::MakeInstance()
+{
+	LLM_SCOPE_BYNAME(TEXT("EnhancedInput/Editor"));
+	return MakeShareable(new FEnhancedInputDeveloperSettingsCustomization());
+}
+
 FEnhancedInputDeveloperSettingsCustomization::~FEnhancedInputDeveloperSettingsCustomization()
 {
 	// Unregister settings panel listeners
@@ -466,6 +472,8 @@ bool FEnhancedInputDeveloperSettingsCustomization::DoesClassHaveSubtypes(UClass*
 
 void FEnhancedInputDeveloperSettingsCustomization::RebuildDetailsViewForAsset(const FAssetData& AssetData, const bool bIsAssetBeingRemoved)
 {
+	LLM_SCOPE_BYNAME(TEXT("EnhancedInput/Editor"));
+	
 	// If the asset was a blueprint...
 	if (AssetData.AssetClassPath == UBlueprint::StaticClass()->GetClassPathName())
 	{
diff --git a/Engine/Plugins/EnhancedInput/Source/InputEditor/Private/InputCustomizations.h b/Engine/Plugins/EnhancedInput/Source/InputEditor/Private/InputCustomizations.h
--- a/Engine/Plugins/EnhancedInput/Source/InputEditor/Private/InputCustomizations.h
+++ b/Engine/Plugins/EnhancedInput/Source/InputEditor/Private/InputCustomizations.h
@@ -59,12 +59,9 @@ class FEnhancedInputDeveloperSettingsCustomization : public IDetailCustomization
 {
 public:
 	//~ IDetailCustomization interface
-	static TSharedRef<IDetailCustomization> MakeInstance()
-	{
-		return MakeShareable(new FEnhancedInputDeveloperSettingsCustomization());
-	}
+	static TSharedRef<IDetailCustomization> MakeInstance();
 
-	virtual ~FEnhancedInputDeveloperSettingsCustomization();
+	virtual ~FEnhancedInputDeveloperSettingsCustomization() override;
 	virtual void CustomizeDetails(IDetailLayoutBuilder& DetailBuilder) override;
 	virtual void CustomizeDetails(const TSharedPtr<IDetailLayoutBuilder>& DetailBuilder) override;
 	//~ End IDetailCustomization interface
diff --git a/Engine/Plugins/EnhancedInput/Source/InputEditor/Private/InputEditorModule.cpp b/Engine/Plugins/EnhancedInput/Source/InputEditor/Private/InputEditorModule.cpp
--- a/Engine/Plugins/EnhancedInput/Source/InputEditor/Private/InputEditorModule.cpp
+++ b/Engine/Plugins/EnhancedInput/Source/InputEditor/Private/InputEditorModule.cpp
@@ -356,6 +356,8 @@ public:
 
 void FInputEditorModule::StartupModule()
 {
+	LLM_SCOPE_BYNAME(TEXT("EnhancedInput/Editor"));
+	
 	// Register customizations
 	FPropertyEditorModule& PropertyModule = FModuleManager::GetModuleChecked<FPropertyEditorModule>("PropertyEditor");
 	PropertyModule.RegisterCustomClassLayout("InputMappingContext", FOnGetDetailCustomizationInstance::CreateStatic(&FInputContextDetails::MakeInstance));
diff --git a/Engine/Source/Runtime/InputCore/Private/InputCoreTypes.cpp b/Engine/Source/Runtime/InputCore/Private/InputCoreTypes.cpp
--- a/Engine/Source/Runtime/InputCore/Private/InputCoreTypes.cpp
+++ b/Engine/Source/Runtime/InputCore/Private/InputCoreTypes.cpp
@@ -9,6 +9,9 @@
 
 DEFINE_LOG_CATEGORY(LogInput);
 
+LLM_DEFINE_TAG(InputCore_EKeys);
+LLM_DEFINE_TAG(InputCore_KeyManager);
+
 #define LOCTEXT_NAMESPACE "InputKeys"
 
 const TCHAR* FKey::SyntheticCharPrefix = TEXT("UnknownCharCode_");
@@ -529,6 +532,9 @@ UE_DISABLE_OPTIMIZATION_SHIP
 void EKeys::Initialize()
 {
 	if (bInitialized) return;
+
+	LLM_SCOPE_BYTAG(InputCore_EKeys);
+	
 	bInitialized = true;
 
 	AddMenuCategoryDisplayInfo(NAME_GamepadCategory, LOCTEXT("GamepadSubCategory", "Gamepad"), TEXT("GraphEditor.PadEvent_16x"));
@@ -992,6 +998,8 @@ UE_ENABLE_OPTIMIZATION_SHIP
 
 void EKeys::AddKey(const FKeyDetails& KeyDetails)
 {
+	LLM_SCOPE_BYTAG(InputCore_EKeys);
+	
 	const FKey& Key = KeyDetails.GetKey();
 	ensureMsgf(!InputKeys.Contains(Key), TEXT("Adding duplicate key '%s'"), *Key.ToString());
 	Key.KeyDetails = MakeShareable(new FKeyDetails(KeyDetails));
@@ -1054,6 +1062,8 @@ TSharedPtr<FKeyDetails> EKeys::GetKeyDetails(const FKey Key)
 
 void EKeys::AddMenuCategoryDisplayInfo(const FName CategoryName, const FText DisplayName, const FName PaletteIcon)
 {
+	LLM_SCOPE_BYTAG(InputCore_EKeys);
+	
 	if (MenuCategoryDisplayInfo.Contains(CategoryName))
 	{
 		UE_LOG(LogInput, Warning, TEXT("Category %s already has menu display info that is being replaced."), *CategoryName.ToString());
@@ -1507,6 +1517,8 @@ TSharedPtr<FInputKeyManager> FInputKeyManager::Instance;
  */
 FInputKeyManager& FInputKeyManager::Get()
 {
+	LLM_SCOPE_BYTAG(InputCore_KeyManager);
+	
 	if( !Instance.IsValid() )
 	{
 		Instance = MakeShareable( new FInputKeyManager() );
@@ -1516,6 +1528,8 @@ FInputKeyManager& FInputKeyManager::Get()
 
 void FInputKeyManager::InitKeyMappings()
 {
+	LLM_SCOPE_BYTAG(InputCore_KeyManager);
+	
 	static const uint32 MAX_KEY_MAPPINGS(256);
 	uint32 KeyCodes[MAX_KEY_MAPPINGS], CharCodes[MAX_KEY_MAPPINGS];
 	FString KeyNames[MAX_KEY_MAPPINGS], CharKeyNames[MAX_KEY_MAPPINGS];
