Add "(Digital)" to the gamepad right and left trigger gamepad key display names in the editor.

This should hopefully help alleviate any confusion when mapping to these keys, since their value is only going to ever be 0.0 or 1.0. In order to get any values in between (i.e. the exact value of the trigger) you should bind to the gamepad trigger "Axis".

Only make these changes in the editor in case there is any gamepad/UI code displaying the name of the keys that game's don't want to change.

#rnx

diff --git a/Engine/Source/Runtime/InputCore/Private/InputCoreTypes.cpp b/Engine/Source/Runtime/InputCore/Private/InputCoreTypes.cpp
--- a/Engine/Source/Runtime/InputCore/Private/InputCoreTypes.cpp
+++ b/Engine/Source/Runtime/InputCore/Private/InputCoreTypes.cpp
@@ -1252,11 +1252,19 @@ FText EKeys::GetGamepadDisplayName(const FKey Key)
 		}
 		else if (Key == EKeys::Gamepad_LeftTrigger)
 		{
+#if WITH_EDITOR
+			return LOCTEXT("Gamepad_LeftTrigger_Editor", "Gamepad Left Trigger (Digital)");
+#else
 			return LOCTEXT("Gamepad_LeftTrigger", "Gamepad Left Trigger");
+#endif
 		}
 		else if (Key == EKeys::Gamepad_RightTrigger)
 		{
+#if WITH_EDITOR
+			return LOCTEXT("Gamepad_RightTrigger_Editor", "Gamepad Right Trigger (Digital)");
+#else
 			return LOCTEXT("Gamepad_RightTrigger", "Gamepad Right Trigger");
+#endif
 		}
 		else if (Key == EKeys::Gamepad_LeftTriggerAxis)
 		{
