Fix to prevent EKeys::TouchKeys array from getting reinitialized with default values in optimized builds due to the FKeyInitializer logic getting executed before the EKeys::TouchKeys variable initialization.

diff --git a/Engine/Source/Runtime/InputCore/Private/InputCoreTypes.cpp b/Engine/Source/Runtime/InputCore/Private/InputCoreTypes.cpp
--- a/Engine/Source/Runtime/InputCore/Private/InputCoreTypes.cpp
+++ b/Engine/Source/Runtime/InputCore/Private/InputCoreTypes.cpp
@@ -214,18 +214,6 @@ const FKey EKeys::Acceleration("Acceleration");
 // Fingers
 const FKey EKeys::TouchKeys[NUM_TOUCH_KEYS];
 
-static struct FKeyInitializer
-{
-	FKeyInitializer()
-	{
-		for (int TouchIndex = 0; TouchIndex < (EKeys::NUM_TOUCH_KEYS - 1); TouchIndex++)
-		{
-			const_cast<FKey&>(EKeys::TouchKeys[TouchIndex]) = FKey(*FString::Printf(TEXT("Touch%d"), TouchIndex + 1));
-		}
-	}
-
-} KeyInitializer;
-
 // Gestures
 const FKey EKeys::Gesture_Pinch("Gesture_Pinch");
 const FKey EKeys::Gesture_Flick("Gesture_Flick");
@@ -537,6 +525,11 @@ void EKeys::Initialize()
 	
 	bInitialized = true;
 
+	for (int32 TouchIndex = 0; TouchIndex < (EKeys::NUM_TOUCH_KEYS - 1); TouchIndex++)
+	{
+		const_cast<FKey&>(EKeys::TouchKeys[TouchIndex]) = FKey(*FString::Printf(TEXT("Touch%d"), TouchIndex + 1));
+	}
+
 	AddMenuCategoryDisplayInfo(NAME_GamepadCategory, LOCTEXT("GamepadSubCategory", "Gamepad"), TEXT("GraphEditor.PadEvent_16x"));
 	AddMenuCategoryDisplayInfo(NAME_MouseCategory, LOCTEXT("MouseSubCategory", "Mouse"), TEXT("GraphEditor.MouseEvent_16x"));
 	AddMenuCategoryDisplayInfo(NAME_KeyboardCategory, LOCTEXT("KeyboardSubCategory", "Keyboard"), TEXT("GraphEditor.KeyEvent_16x"));
