Gamepad touchpad touch can be bound as a simple button (without clicking it in)

#tests PIE

diff --git a/Engine/Source/Runtime/ApplicationCore/Private/GenericPlatform/GenericApplication.cpp b/Engine/Source/Runtime/ApplicationCore/Private/GenericPlatform/GenericApplication.cpp
--- a/Engine/Source/Runtime/ApplicationCore/Private/GenericPlatform/GenericApplication.cpp
+++ b/Engine/Source/Runtime/ApplicationCore/Private/GenericPlatform/GenericApplication.cpp
@@ -21,6 +21,7 @@ const FGamepadKeyNames::Type FGamepadKeyNames::RightThumb("Gamepad_RightThumbsti
 const FGamepadKeyNames::Type FGamepadKeyNames::SpecialLeft("Gamepad_Special_Left");
 const FGamepadKeyNames::Type FGamepadKeyNames::SpecialLeft_X("Gamepad_Special_Left_X");
 const FGamepadKeyNames::Type FGamepadKeyNames::SpecialLeft_Y("Gamepad_Special_Left_Y");
+const FGamepadKeyNames::Type FGamepadKeyNames::SpecialLeft_Touched("Gamepad_Special_Left_Touched");
 const FGamepadKeyNames::Type FGamepadKeyNames::SpecialRight("Gamepad_Special_Right");
 const FGamepadKeyNames::Type FGamepadKeyNames::FaceButtonBottom("Gamepad_FaceButton_Bottom");
 const FGamepadKeyNames::Type FGamepadKeyNames::FaceButtonRight("Gamepad_FaceButton_Right");
diff --git a/Engine/Source/Runtime/ApplicationCore/Public/GenericPlatform/GenericApplicationMessageHandler.h b/Engine/Source/Runtime/ApplicationCore/Public/GenericPlatform/GenericApplicationMessageHandler.h
--- a/Engine/Source/Runtime/ApplicationCore/Public/GenericPlatform/GenericApplicationMessageHandler.h
+++ b/Engine/Source/Runtime/ApplicationCore/Public/GenericPlatform/GenericApplicationMessageHandler.h
@@ -44,6 +44,7 @@ struct FGamepadKeyNames
 	static APPLICATIONCORE_API const FName SpecialLeft;
 	static APPLICATIONCORE_API const FName SpecialLeft_X;
 	static APPLICATIONCORE_API const FName SpecialLeft_Y;
+	static APPLICATIONCORE_API const FName SpecialLeft_Touched;
 	static APPLICATIONCORE_API const FName SpecialRight;
 	static APPLICATIONCORE_API const FName FaceButtonBottom;
 	static APPLICATIONCORE_API const FName FaceButtonRight;
diff --git a/Engine/Source/Runtime/InputCore/Classes/InputCoreTypes.h b/Engine/Source/Runtime/InputCore/Classes/InputCoreTypes.h
--- a/Engine/Source/Runtime/InputCore/Classes/InputCoreTypes.h
+++ b/Engine/Source/Runtime/InputCore/Classes/InputCoreTypes.h
@@ -440,6 +440,7 @@ struct EKeys
 	static INPUTCORE_API const FKey Gamepad_Special_Left;
 	static INPUTCORE_API const FKey Gamepad_Special_Left_X;
 	static INPUTCORE_API const FKey Gamepad_Special_Left_Y;
+	static INPUTCORE_API const FKey Gamepad_Special_Left_Touched;
 	static INPUTCORE_API const FKey Gamepad_Special_Right;
 	static INPUTCORE_API const FKey Gamepad_FaceButton_Bottom;
 	static INPUTCORE_API const FKey Gamepad_FaceButton_Right;
diff --git a/Engine/Source/Runtime/InputCore/Private/InputCoreTypes.cpp b/Engine/Source/Runtime/InputCore/Private/InputCoreTypes.cpp
--- a/Engine/Source/Runtime/InputCore/Private/InputCoreTypes.cpp
+++ b/Engine/Source/Runtime/InputCore/Private/InputCoreTypes.cpp
@@ -180,6 +180,7 @@ const FKey EKeys::Gamepad_RightThumbstick("Gamepad_RightThumbstick");
 const FKey EKeys::Gamepad_Special_Left("Gamepad_Special_Left");
 const FKey EKeys::Gamepad_Special_Left_X("Gamepad_Special_Left_X");
 const FKey EKeys::Gamepad_Special_Left_Y("Gamepad_Special_Left_Y");
+const FKey EKeys::Gamepad_Special_Left_Touched("Gamepad_Special_Left_Touched");
 const FKey EKeys::Gamepad_Special_Right("Gamepad_Special_Right");
 const FKey EKeys::Gamepad_FaceButton_Bottom("Gamepad_FaceButton_Bottom");
 const FKey EKeys::Gamepad_FaceButton_Right("Gamepad_FaceButton_Right");
@@ -761,6 +762,7 @@ void EKeys::Initialize()
 	PRAGMA_ENABLE_DEPRECATION_WARNINGS
 	AddKey(FKeyDetails(EKeys::Gamepad_Special_Left_X, LOCTEXT("Gamepad_Special_Left_X", "Touchpad Button X Axis"), FKeyDetails::GamepadKey | FKeyDetails::Axis1D, "PS4"));
 	AddKey(FKeyDetails(EKeys::Gamepad_Special_Left_Y, LOCTEXT("Gamepad_Special_Left_Y", "Touchpad Button Y Axis"), FKeyDetails::GamepadKey | FKeyDetails::Axis1D, "PS4"));
+	AddKey(FKeyDetails(EKeys::Gamepad_Special_Left_Touched, LOCTEXT("Gamepad_Special_Left_Touched", "Touchpad Touched"), FKeyDetails::GamepadKey, "PS4"));
 
 
 	// Steam Controller specific
@@ -1113,6 +1115,10 @@ FText EKeys::GetGamepadDisplayName(const FKey Key)
 		{
 			return LOCTEXT("PS4_Gamepad_Special_Left", "Gamepad Touchpad Button");
 		}
+		else if (Key == EKeys::Gamepad_Special_Left_Touched)
+		{
+			return LOCTEXT("PS4_Gamepad_Special_Left_Touched", "Gamepad Touchpad Touched");
+		}
 		else if (Key == EKeys::Gamepad_Special_Right)
 		{
 			return LOCTEXT("PS4_Gamepad_Special_Right", "Gamepad Options");
