[Input] Remove some more deprecated FKey types.

diff --git a/Engine/Source/Programs/AutomationTool/Gauntlet/Framework/Utils/Gauntlet.InputActions.cs b/Engine/Source/Programs/AutomationTool/Gauntlet/Framework/Utils/Gauntlet.InputActions.cs
--- a/Engine/Source/Programs/AutomationTool/Gauntlet/Framework/Utils/Gauntlet.InputActions.cs
+++ b/Engine/Source/Programs/AutomationTool/Gauntlet/Framework/Utils/Gauntlet.InputActions.cs
@@ -334,7 +334,6 @@ namespace Gauntlet.Utils
 		public static FKey Gamepad_FaceButton_Top = new FKey("Gamepad_FaceButton_Top");
 		public static FKey Gamepad_LeftShoulder = new FKey("Gamepad_LeftShoulder");
 		public static FKey Gamepad_RightShoulder = new FKey("Gamepad_RightShoulder");
-		public static FKey PS4_Special = new FKey("PS4_Special");
 
 		// Motion Controls
 		public static FKey Tilt = new FKey("Tilt");
@@ -369,7 +368,6 @@ namespace Gauntlet.Utils
 		public static FKey Android_Menu = new FKey("Android_Menu");
 
 		// Vive Controls
-		public static FKey Vive_Left_System_Click = new FKey("Vive_Left_System_Click");
 		public static FKey Vive_Left_Grip_Click = new FKey("Vive_Left_Grip_Click");
 		public static FKey Vive_Left_Menu_Click = new FKey("Vive_Left_Menu_Click");
 		public static FKey Vive_Left_Trigger_Click = new FKey("Vive_Left_Trigger_Click");
@@ -383,7 +381,6 @@ namespace Gauntlet.Utils
 		public static FKey Vive_Left_Trackpad_Down = new FKey("Vive_Left_Trackpad_Down");
 		public static FKey Vive_Left_Trackpad_Left = new FKey("Vive_Left_Trackpad_Left");
 		public static FKey Vive_Left_Trackpad_Right = new FKey("Vive_Left_Trackpad_Right");
-		public static FKey Vive_Right_System_Click = new FKey("Vive_Right_System_Click");
 		public static FKey Vive_Right_Grip_Click = new FKey("Vive_Right_Grip_Click");
 		public static FKey Vive_Right_Menu_Click = new FKey("Vive_Right_Menu_Click");
 		public static FKey Vive_Right_Trigger_Click = new FKey("Vive_Right_Trigger_Click");
@@ -466,7 +463,6 @@ namespace Gauntlet.Utils
 		public static FKey OculusTouch_Right_B_Click = new FKey("OculusTouch_Right_B_Click");
 		public static FKey OculusTouch_Right_A_Touch = new FKey("OculusTouch_Right_A_Touch");
 		public static FKey OculusTouch_Right_B_Touch = new FKey("OculusTouch_Right_B_Touch");
-		public static FKey OculusTouch_Right_System_Click = new FKey("OculusTouch_Right_System_Click");
 		public static FKey OculusTouch_Right_Grip_Click = new FKey("OculusTouch_Right_Grip_Click");
 		public static FKey OculusTouch_Right_Grip_Axis = new FKey("OculusTouch_Right_Grip_Axis");
 		public static FKey OculusTouch_Right_Trigger_Click = new FKey("OculusTouch_Right_Trigger_Click");
@@ -487,8 +483,6 @@ namespace Gauntlet.Utils
 		public static FKey ValveIndex_Left_B_Click = new FKey("ValveIndex_Left_B_Click");
 		public static FKey ValveIndex_Left_A_Touch = new FKey("ValveIndex_Left_A_Touch");
 		public static FKey ValveIndex_Left_B_Touch = new FKey("ValveIndex_Left_B_Touch");
-		public static FKey ValveIndex_Left_System_Click = new FKey("ValveIndex_Left_System_Click");
-		public static FKey ValveIndex_Left_System_Touch = new FKey("ValveIndex_Left_System_Touch");
 		public static FKey ValveIndex_Left_Grip_Axis = new FKey("ValveIndex_Left_Grip_Axis");
 		public static FKey ValveIndex_Left_Grip_Force = new FKey("ValveIndex_Left_Grip_Force");
 		public static FKey ValveIndex_Left_Trigger_Click = new FKey("ValveIndex_Left_Trigger_Click");
@@ -516,8 +510,6 @@ namespace Gauntlet.Utils
 		public static FKey ValveIndex_Right_B_Click = new FKey("ValveIndex_Right_B_Click");
 		public static FKey ValveIndex_Right_A_Touch = new FKey("ValveIndex_Right_A_Touch");
 		public static FKey ValveIndex_Right_B_Touch = new FKey("ValveIndex_Right_B_Touch");
-		public static FKey ValveIndex_Right_System_Click = new FKey("ValveIndex_Right_System_Click");
-		public static FKey ValveIndex_Right_System_Touch = new FKey("ValveIndex_Right_System_Touch");
 		public static FKey ValveIndex_Right_Grip_Axis = new FKey("ValveIndex_Right_Grip_Axis");
 		public static FKey ValveIndex_Right_Grip_Force = new FKey("ValveIndex_Right_Grip_Force");
 		public static FKey ValveIndex_Right_Trigger_Click = new FKey("ValveIndex_Right_Trigger_Click");
diff --git a/Engine/Source/Runtime/InputCore/Classes/InputCoreTypes.h b/Engine/Source/Runtime/InputCore/Classes/InputCoreTypes.h
--- a/Engine/Source/Runtime/InputCore/Classes/InputCoreTypes.h
+++ b/Engine/Source/Runtime/InputCore/Classes/InputCoreTypes.h
@@ -464,10 +464,6 @@ struct EKeys
 	static INPUTCORE_API const FKey Gesture_Flick;
 	static INPUTCORE_API const FKey Gesture_Rotate;
 
-	// PS4-specific
-	UE_DEPRECATED(5.0, "This key has deprecated and will be removed. Use GamePad_Special_Left/Right instead.")
-	static INPUTCORE_API const FKey PS4_Special;
-
 	// Steam Controller Specific
 	static INPUTCORE_API const FKey Steam_Touch_0;
 	static INPUTCORE_API const FKey Steam_Touch_1;
@@ -490,8 +486,6 @@ struct EKeys
 	static INPUTCORE_API const FKey Android_Menu;
 
 	// HTC Vive Controller
-	UE_DEPRECATED(5.1, "This key has been deprecated and will be removed.")
-	static INPUTCORE_API const FKey Vive_Left_System_Click;
 	static INPUTCORE_API const FKey Vive_Left_Grip_Click;
 	static INPUTCORE_API const FKey Vive_Left_Menu_Click;
 	static INPUTCORE_API const FKey Vive_Left_Trigger_Click;
@@ -505,8 +499,6 @@ struct EKeys
 	static INPUTCORE_API const FKey Vive_Left_Trackpad_Down;
 	static INPUTCORE_API const FKey Vive_Left_Trackpad_Left;
 	static INPUTCORE_API const FKey Vive_Left_Trackpad_Right;
-	UE_DEPRECATED(5.1, "This key has been deprecated and will be removed.")
-	static INPUTCORE_API const FKey Vive_Right_System_Click;
 	static INPUTCORE_API const FKey Vive_Right_Grip_Click;
 	static INPUTCORE_API const FKey Vive_Right_Menu_Click;
 	static INPUTCORE_API const FKey Vive_Right_Trigger_Click;
@@ -589,8 +581,6 @@ struct EKeys
 	static INPUTCORE_API const FKey OculusTouch_Right_B_Click;
 	static INPUTCORE_API const FKey OculusTouch_Right_A_Touch;
 	static INPUTCORE_API const FKey OculusTouch_Right_B_Touch;
-	UE_DEPRECATED(5.1, "This key has been deprecated and will be removed.")
-	static INPUTCORE_API const FKey OculusTouch_Right_System_Click;
 	static INPUTCORE_API const FKey OculusTouch_Right_Grip_Click;
 	static INPUTCORE_API const FKey OculusTouch_Right_Grip_Axis;
 	static INPUTCORE_API const FKey OculusTouch_Right_Trigger_Click;
@@ -611,10 +601,6 @@ struct EKeys
 	static INPUTCORE_API const FKey ValveIndex_Left_B_Click;
 	static INPUTCORE_API const FKey ValveIndex_Left_A_Touch;
 	static INPUTCORE_API const FKey ValveIndex_Left_B_Touch;
-	UE_DEPRECATED(5.1, "This key has been deprecated and will be removed.")
-	static INPUTCORE_API const FKey ValveIndex_Left_System_Click;
-	UE_DEPRECATED(5.1, "This key has been deprecated and will be removed.")
-	static INPUTCORE_API const FKey ValveIndex_Left_System_Touch;
 	static INPUTCORE_API const FKey ValveIndex_Left_Grip_Axis;
 	static INPUTCORE_API const FKey ValveIndex_Left_Grip_Force;
 	static INPUTCORE_API const FKey ValveIndex_Left_Trigger_Click;
@@ -642,10 +628,6 @@ struct EKeys
 	static INPUTCORE_API const FKey ValveIndex_Right_B_Click;
 	static INPUTCORE_API const FKey ValveIndex_Right_A_Touch;
 	static INPUTCORE_API const FKey ValveIndex_Right_B_Touch;
-	UE_DEPRECATED(5.1, "This key has been deprecated and will be removed.")
-	static INPUTCORE_API const FKey ValveIndex_Right_System_Click;
-	UE_DEPRECATED(5.1, "This key has been deprecated and will be removed.")
-	static INPUTCORE_API const FKey ValveIndex_Right_System_Touch;
 	static INPUTCORE_API const FKey ValveIndex_Right_Grip_Axis;
 	static INPUTCORE_API const FKey ValveIndex_Right_Grip_Force;
 	static INPUTCORE_API const FKey ValveIndex_Right_Trigger_Click;
diff --git a/Engine/Source/Runtime/InputCore/Private/InputCoreTypes.cpp b/Engine/Source/Runtime/InputCore/Private/InputCoreTypes.cpp
--- a/Engine/Source/Runtime/InputCore/Private/InputCoreTypes.cpp
+++ b/Engine/Source/Runtime/InputCore/Private/InputCoreTypes.cpp
@@ -220,11 +220,6 @@ const FKey EKeys::Gesture_Pinch("Gesture_Pinch");
 const FKey EKeys::Gesture_Flick("Gesture_Flick");
 const FKey EKeys::Gesture_Rotate("Gesture_Rotate");
 
-PRAGMA_DISABLE_DEPRECATION_WARNINGS
-// PS4-specific
-const FKey EKeys::PS4_Special("PS4_Special");
-PRAGMA_ENABLE_DEPRECATION_WARNINGS
-
 // Steam Controller Specific
 const FKey EKeys::Steam_Touch_0("Steam_Touch_0");
 const FKey EKeys::Steam_Touch_1("Steam_Touch_1");
@@ -247,9 +242,6 @@ const FKey EKeys::Android_Volume_Down("Android_Volume_Down");
 const FKey EKeys::Android_Menu("Android_Menu");
 
 // HTC Vive Controller
-PRAGMA_DISABLE_DEPRECATION_WARNINGS
-const FKey EKeys::Vive_Left_System_Click("Vive_Left_System_Click");
-PRAGMA_ENABLE_DEPRECATION_WARNINGS
 const FKey EKeys::Vive_Left_Grip_Click("Vive_Left_Grip_Click");
 const FKey EKeys::Vive_Left_Menu_Click("Vive_Left_Menu_Click");
 const FKey EKeys::Vive_Left_Trigger_Click("Vive_Left_Trigger_Click");
@@ -263,9 +255,6 @@ const FKey EKeys::Vive_Left_Trackpad_Up("Vive_Left_Trackpad_Up");
 const FKey EKeys::Vive_Left_Trackpad_Down("Vive_Left_Trackpad_Down");
 const FKey EKeys::Vive_Left_Trackpad_Left("Vive_Left_Trackpad_Left");
 const FKey EKeys::Vive_Left_Trackpad_Right("Vive_Left_Trackpad_Right");
-PRAGMA_DISABLE_DEPRECATION_WARNINGS
-const FKey EKeys::Vive_Right_System_Click("Vive_Right_System_Click");
-PRAGMA_ENABLE_DEPRECATION_WARNINGS
 const FKey EKeys::Vive_Right_Grip_Click("Vive_Right_Grip_Click");
 const FKey EKeys::Vive_Right_Menu_Click("Vive_Right_Menu_Click");
 const FKey EKeys::Vive_Right_Trigger_Click("Vive_Right_Trigger_Click");
@@ -348,9 +337,6 @@ const FKey EKeys::OculusTouch_Right_A_Click("OculusTouch_Right_A_Click");
 const FKey EKeys::OculusTouch_Right_B_Click("OculusTouch_Right_B_Click");
 const FKey EKeys::OculusTouch_Right_A_Touch("OculusTouch_Right_A_Touch");
 const FKey EKeys::OculusTouch_Right_B_Touch("OculusTouch_Right_B_Touch");
-PRAGMA_DISABLE_DEPRECATION_WARNINGS
-const FKey EKeys::OculusTouch_Right_System_Click("OculusTouch_Right_System_Click");
-PRAGMA_ENABLE_DEPRECATION_WARNINGS
 const FKey EKeys::OculusTouch_Right_Grip_Click("OculusTouch_Right_Grip_Click");
 const FKey EKeys::OculusTouch_Right_Grip_Axis("OculusTouch_Right_Grip_Axis");
 const FKey EKeys::OculusTouch_Right_Trigger_Click("OculusTouch_Right_Trigger_Click");
@@ -371,10 +357,6 @@ const FKey EKeys::ValveIndex_Left_A_Click("ValveIndex_Left_A_Click");
 const FKey EKeys::ValveIndex_Left_B_Click("ValveIndex_Left_B_Click");
 const FKey EKeys::ValveIndex_Left_A_Touch("ValveIndex_Left_A_Touch");
 const FKey EKeys::ValveIndex_Left_B_Touch("ValveIndex_Left_B_Touch");
-PRAGMA_DISABLE_DEPRECATION_WARNINGS
-const FKey EKeys::ValveIndex_Left_System_Click("ValveIndex_Left_System_Click");
-const FKey EKeys::ValveIndex_Left_System_Touch("ValveIndex_Left_System_Touch");
-PRAGMA_ENABLE_DEPRECATION_WARNINGS
 const FKey EKeys::ValveIndex_Left_Grip_Axis("ValveIndex_Left_Grip_Axis");
 const FKey EKeys::ValveIndex_Left_Grip_Force("ValveIndex_Left_Grip_Force");
 const FKey EKeys::ValveIndex_Left_Trigger_Click("ValveIndex_Left_Trigger_Click");
@@ -402,10 +384,6 @@ const FKey EKeys::ValveIndex_Right_A_Click("ValveIndex_Right_A_Click");
 const FKey EKeys::ValveIndex_Right_B_Click("ValveIndex_Right_B_Click");
 const FKey EKeys::ValveIndex_Right_A_Touch("ValveIndex_Right_A_Touch");
 const FKey EKeys::ValveIndex_Right_B_Touch("ValveIndex_Right_B_Touch");
-PRAGMA_DISABLE_DEPRECATION_WARNINGS
-const FKey EKeys::ValveIndex_Right_System_Click("ValveIndex_Right_System_Click");
-const FKey EKeys::ValveIndex_Right_System_Touch("ValveIndex_Right_System_Touch");
-PRAGMA_ENABLE_DEPRECATION_WARNINGS
 const FKey EKeys::ValveIndex_Right_Grip_Axis("ValveIndex_Right_Grip_Axis");
 const FKey EKeys::ValveIndex_Right_Grip_Force("ValveIndex_Right_Grip_Force");
 const FKey EKeys::ValveIndex_Right_Trigger_Click("ValveIndex_Right_Trigger_Click");
@@ -758,12 +736,9 @@ void EKeys::Initialize()
 
 	// PS4-specific
 	AddMenuCategoryDisplayInfo("Special Gamepad", LOCTEXT("SpecialGamepadSubCategory", "SpecialGamepad"), TEXT("GraphEditor.PadEvent_16x"));
-	PRAGMA_DISABLE_DEPRECATION_WARNINGS
-	AddKey(FKeyDetails(EKeys::PS4_Special, LOCTEXT("PS4_Special", "PS4_Special_DEPRECATED"), FKeyDetails::Deprecated | FKeyDetails::NotBlueprintBindableKey | FKeyDetails::NotActionBindableKey, "PS4"));
-	PRAGMA_ENABLE_DEPRECATION_WARNINGS
-	AddKey(FKeyDetails(EKeys::Gamepad_Special_Left_X, LOCTEXT("Gamepad_Special_Left_X", "Touchpad Button X Axis"), FKeyDetails::GamepadKey | FKeyDetails::Axis1D, "PS4"));
-	AddKey(FKeyDetails(EKeys::Gamepad_Special_Left_Y, LOCTEXT("Gamepad_Special_Left_Y", "Touchpad Button Y Axis"), FKeyDetails::GamepadKey | FKeyDetails::Axis1D, "PS4"));
-	AddKey(FKeyDetails(EKeys::Gamepad_Special_Left_Touched, LOCTEXT("Gamepad_Special_Left_Touched", "Touchpad Touched"), FKeyDetails::GamepadKey, "PS4"));
+	AddKey(FKeyDetails(EKeys::Gamepad_Special_Left_X, LOCTEXT("Gamepad_Special_Left_X", "Touchpad Button X Axis"), FKeyDetails::GamepadKey | FKeyDetails::Axis1D, "Special"));
+	AddKey(FKeyDetails(EKeys::Gamepad_Special_Left_Y, LOCTEXT("Gamepad_Special_Left_Y", "Touchpad Button Y Axis"), FKeyDetails::GamepadKey | FKeyDetails::Axis1D, "Special"));
+	AddKey(FKeyDetails(EKeys::Gamepad_Special_Left_Touched, LOCTEXT("Gamepad_Special_Left_Touched", "Touchpad Touched"), FKeyDetails::GamepadKey, "Special"));
 
 
 	// Steam Controller specific
@@ -795,10 +770,7 @@ void EKeys::Initialize()
 
 	// HTC Vive Controller
 	AddMenuCategoryDisplayInfo("Vive", LOCTEXT("ViveSubCategory", "HTC Vive"), TEXT("GraphEditor.PadEvent_16x"));
-
-	PRAGMA_DISABLE_DEPRECATION_WARNINGS
-	AddKey(FKeyDetails(EKeys::Vive_Left_System_Click, LOCTEXT("Vive_Left_System_Click", "Vive (L) System"), FKeyDetails::Deprecated | FKeyDetails::GamepadKey | FKeyDetails::NotBlueprintBindableKey, "Vive"));
-	PRAGMA_ENABLE_DEPRECATION_WARNINGS
+	
 	AddKey(FKeyDetails(EKeys::Vive_Left_Grip_Click, LOCTEXT("Vive_Left_Grip_Click", "Vive (L) Grip"), FKeyDetails::GamepadKey | FKeyDetails::NotBlueprintBindableKey, "Vive"));
 	AddKey(FKeyDetails(EKeys::Vive_Left_Menu_Click, LOCTEXT("Vive_Left_Menu_Click", "Vive (L) Menu"), FKeyDetails::GamepadKey | FKeyDetails::NotBlueprintBindableKey, "Vive"));
 	AddKey(FKeyDetails(EKeys::Vive_Left_Trigger_Click, LOCTEXT("Vive_Left_Trigger_Click", "Vive (L) Trigger"), FKeyDetails::GamepadKey | FKeyDetails::NotBlueprintBindableKey, "Vive"));
@@ -812,9 +784,6 @@ void EKeys::Initialize()
 	AddKey(FKeyDetails(EKeys::Vive_Left_Trackpad_Down, LOCTEXT("Vive_Left_Trackpad_Down", "Vive (L) Trackpad Down"), FKeyDetails::GamepadKey | FKeyDetails::NotBlueprintBindableKey, "Vive"));
 	AddKey(FKeyDetails(EKeys::Vive_Left_Trackpad_Left, LOCTEXT("Vive_Left_Trackpad_Left", "Vive (L) Trackpad Left"), FKeyDetails::GamepadKey | FKeyDetails::NotBlueprintBindableKey, "Vive"));
 	AddKey(FKeyDetails(EKeys::Vive_Left_Trackpad_Right, LOCTEXT("Vive_Left_Trackpad_Right", "Vive (L) Trackpad Right"), FKeyDetails::GamepadKey | FKeyDetails::NotBlueprintBindableKey, "Vive"));
-	PRAGMA_DISABLE_DEPRECATION_WARNINGS
-	AddKey(FKeyDetails(EKeys::Vive_Right_System_Click, LOCTEXT("Vive_Right_System_Click", "Vive (R) System"), FKeyDetails::Deprecated | FKeyDetails::GamepadKey | FKeyDetails::NotBlueprintBindableKey, "Vive"));
-	PRAGMA_ENABLE_DEPRECATION_WARNINGS
 	AddKey(FKeyDetails(EKeys::Vive_Right_Grip_Click, LOCTEXT("Vive_Right_Grip_Click", "Vive (R) Grip"), FKeyDetails::GamepadKey | FKeyDetails::NotBlueprintBindableKey, "Vive"));
 	AddKey(FKeyDetails(EKeys::Vive_Right_Menu_Click, LOCTEXT("Vive_Right_Menu_Click", "Vive (R) Menu"), FKeyDetails::GamepadKey | FKeyDetails::NotBlueprintBindableKey, "Vive"));
 	AddKey(FKeyDetails(EKeys::Vive_Right_Trigger_Click, LOCTEXT("Vive_Right_Trigger_Click", "Vive (R) Trigger"), FKeyDetails::GamepadKey | FKeyDetails::NotBlueprintBindableKey, "Vive"));
@@ -901,9 +870,6 @@ void EKeys::Initialize()
 	AddKey(FKeyDetails(EKeys::OculusTouch_Right_B_Click, LOCTEXT("OculusTouch_Right_B_Click", "Oculus Touch (R) B Press"), FKeyDetails::GamepadKey | FKeyDetails::NotBlueprintBindableKey, "OculusTouch"));
 	AddKey(FKeyDetails(EKeys::OculusTouch_Right_A_Touch, LOCTEXT("OculusTouch_Right_A_Touch", "Oculus Touch (R) A Touch"), FKeyDetails::GamepadKey | FKeyDetails::NotBlueprintBindableKey, "OculusTouch"));
 	AddKey(FKeyDetails(EKeys::OculusTouch_Right_B_Touch, LOCTEXT("OculusTouch_Right_B_Touch", "Oculus Touch (R) B Touch"), FKeyDetails::GamepadKey | FKeyDetails::NotBlueprintBindableKey, "OculusTouch"));
-	PRAGMA_DISABLE_DEPRECATION_WARNINGS
-	AddKey(FKeyDetails(EKeys::OculusTouch_Right_System_Click, LOCTEXT("OculusTouch_Right_System_Click", "Oculus Touch (R) System"), FKeyDetails::Deprecated | FKeyDetails::GamepadKey | FKeyDetails::NotBlueprintBindableKey, "OculusTouch"));
-	PRAGMA_ENABLE_DEPRECATION_WARNINGS
 	AddKey(FKeyDetails(EKeys::OculusTouch_Right_Grip_Click, LOCTEXT("OculusTouch_Right_Grip_Click", "Oculus Touch (R) Grip"), FKeyDetails::GamepadKey | FKeyDetails::NotBlueprintBindableKey, "OculusTouch"));
 	AddKey(FKeyDetails(EKeys::OculusTouch_Right_Grip_Axis, LOCTEXT("OculusTouch_Right_Grip_Axis", "Oculus Touch (R) Grip Axis"), FKeyDetails::GamepadKey | FKeyDetails::Axis1D | FKeyDetails::NotBlueprintBindableKey, "OculusTouch"));
 	AddKey(FKeyDetails(EKeys::OculusTouch_Right_Trigger_Click, LOCTEXT("OculusTouch_Right_Trigger_Click", "Oculus Touch (R) Trigger"), FKeyDetails::GamepadKey | FKeyDetails::NotBlueprintBindableKey, "OculusTouch"));
@@ -926,10 +892,6 @@ void EKeys::Initialize()
 	AddKey(FKeyDetails(EKeys::ValveIndex_Left_B_Click, LOCTEXT("ValveIndex_Left_B_Click", "Valve Index (L) B Press"), FKeyDetails::GamepadKey | FKeyDetails::NotBlueprintBindableKey, "ValveIndex"));
 	AddKey(FKeyDetails(EKeys::ValveIndex_Left_A_Touch, LOCTEXT("ValveIndex_Left_A_Touch", "Valve Index (L) A Touch"), FKeyDetails::GamepadKey | FKeyDetails::NotBlueprintBindableKey, "ValveIndex"));
 	AddKey(FKeyDetails(EKeys::ValveIndex_Left_B_Touch, LOCTEXT("ValveIndex_Left_B_Touch", "Valve Index (L) B Touch"), FKeyDetails::GamepadKey | FKeyDetails::NotBlueprintBindableKey, "ValveIndex"));
-	PRAGMA_DISABLE_DEPRECATION_WARNINGS
-	AddKey(FKeyDetails(EKeys::ValveIndex_Left_System_Click, LOCTEXT("ValveIndex_Left_System_Click", "Valve Index (L) System"), FKeyDetails::Deprecated | FKeyDetails::GamepadKey | FKeyDetails::NotBlueprintBindableKey, "ValveIndex"));
-	AddKey(FKeyDetails(EKeys::ValveIndex_Left_System_Touch, LOCTEXT("ValveIndex_Left_System_Touch", "Valve Index (L) System Touch"), FKeyDetails::Deprecated | FKeyDetails::GamepadKey | FKeyDetails::NotBlueprintBindableKey, "ValveIndex"));
-	PRAGMA_ENABLE_DEPRECATION_WARNINGS
 	AddKey(FKeyDetails(EKeys::ValveIndex_Left_Grip_Axis, LOCTEXT("ValveIndex_Left_Grip_Axis", "Valve Index (L) Grip Axis"), FKeyDetails::GamepadKey | FKeyDetails::Axis1D | FKeyDetails::NotBlueprintBindableKey, "ValveIndex"));
 	AddKey(FKeyDetails(EKeys::ValveIndex_Left_Grip_Force, LOCTEXT("ValveIndex_Left_Grip_Force", "Valve Index (L) Grip Force"), FKeyDetails::GamepadKey | FKeyDetails::Axis1D | FKeyDetails::NotBlueprintBindableKey, "ValveIndex"));
 	AddKey(FKeyDetails(EKeys::ValveIndex_Left_Trigger_Click, LOCTEXT("ValveIndex_Left_Trigger_Click", "Valve Index (L) Trigger"), FKeyDetails::GamepadKey | FKeyDetails::NotBlueprintBindableKey, "ValveIndex"));
@@ -957,10 +919,6 @@ void EKeys::Initialize()
 	AddKey(FKeyDetails(EKeys::ValveIndex_Right_B_Click, LOCTEXT("ValveIndex_Right_B_Click", "Valve Index (R) B Press"), FKeyDetails::GamepadKey | FKeyDetails::NotBlueprintBindableKey, "ValveIndex"));
 	AddKey(FKeyDetails(EKeys::ValveIndex_Right_A_Touch, LOCTEXT("ValveIndex_Right_A_Touch", "Valve Index (R) A Touch"), FKeyDetails::GamepadKey | FKeyDetails::NotBlueprintBindableKey, "ValveIndex"));
 	AddKey(FKeyDetails(EKeys::ValveIndex_Right_B_Touch, LOCTEXT("ValveIndex_Right_B_Touch", "Valve Index (R) B Touch"), FKeyDetails::GamepadKey | FKeyDetails::NotBlueprintBindableKey, "ValveIndex"));
-	PRAGMA_DISABLE_DEPRECATION_WARNINGS
-	AddKey(FKeyDetails(EKeys::ValveIndex_Right_System_Click, LOCTEXT("ValveIndex_Right_System_Click", "Valve Index (R) System"), FKeyDetails::Deprecated | FKeyDetails::GamepadKey | FKeyDetails::NotBlueprintBindableKey, "ValveIndex"));
-	AddKey(FKeyDetails(EKeys::ValveIndex_Right_System_Touch, LOCTEXT("ValveIndex_Right_System_Touch", "Valve Index (R) System Touch"), FKeyDetails::Deprecated | FKeyDetails::GamepadKey | FKeyDetails::NotBlueprintBindableKey, "ValveIndex"));
-	PRAGMA_ENABLE_DEPRECATION_WARNINGS
 	AddKey(FKeyDetails(EKeys::ValveIndex_Right_Grip_Axis, LOCTEXT("ValveIndex_Right_Grip_Axis", "Valve Index (R) Grip Axis"), FKeyDetails::GamepadKey | FKeyDetails::Axis1D | FKeyDetails::NotBlueprintBindableKey, "ValveIndex"));
 	AddKey(FKeyDetails(EKeys::ValveIndex_Right_Grip_Force, LOCTEXT("ValveIndex_Right_Grip_Force", "Valve Index (R) Grip Force"), FKeyDetails::GamepadKey | FKeyDetails::Axis1D | FKeyDetails::NotBlueprintBindableKey, "ValveIndex"));
 	AddKey(FKeyDetails(EKeys::ValveIndex_Right_Trigger_Click, LOCTEXT("ValveIndex_Right_Trigger_Click", "Valve Index (R) Trigger"), FKeyDetails::GamepadKey | FKeyDetails::NotBlueprintBindableKey, "ValveIndex"));
