Restored the Ctrl keys for Android.

diff --git a/Engine/Source/Runtime/InputCore/Private/Android/AndroidPlatformInput.cpp b/Engine/Source/Runtime/InputCore/Private/Android/AndroidPlatformInput.cpp
--- a/Engine/Source/Runtime/InputCore/Private/Android/AndroidPlatformInput.cpp
+++ b/Engine/Source/Runtime/InputCore/Private/Android/AndroidPlatformInput.cpp
@@ -253,8 +253,8 @@ uint32 FAndroidPlatformInput::GetKeyMap( uint32* KeyCodes, FString* KeyNames, ui
 
 		ADDKEYMAP( AKEYCODE_SHIFT_LEFT, TEXT("LeftShift") );
 		ADDKEYMAP( AKEYCODE_SHIFT_RIGHT, TEXT("RightShift") );
-		//ADDKEYMAP( VK_LCONTROL, TEXT("LeftControl") );
-		//ADDKEYMAP( VK_RCONTROL, TEXT("RightControl") );
+		ADDKEYMAP( AKEYCODE_CTRL_LEFT, TEXT("LeftControl") );
+		ADDKEYMAP( AKEYCODE_CTRL_RIGHT, TEXT("RightControl") );
 		ADDKEYMAP( AKEYCODE_ALT_LEFT, TEXT("LeftAlt") );
 		ADDKEYMAP( AKEYCODE_ALT_RIGHT, TEXT("RightAlt") );
 		//ADDKEYMAP( VK_LWIN, TEXT("LeftCommand") );
