[Input] Create a "proper" concept of Virtual FKeys at the input core level.

- Add a "IsVirtual" flag to FKeys so that you check at runtime if the key is virtual.
- Add two new virtual keys for accept and back, which are bindable and useable in Enhanced Input. We need to make new keys here because FKeys are serialized via their FName, so in order to make sure that it properly serializes per-platform we need to save out the keys with unique names and then populate the key details with the "real" value at runtime.

[RN] Add bindable "virtual" accept and back FKeys.

diff --git a/Engine/Plugins/EnhancedInput/Source/EnhancedInput/Private/EnhancedPlayerInput.cpp b/Engine/Plugins/EnhancedInput/Source/EnhancedInput/Private/EnhancedPlayerInput.cpp
--- a/Engine/Plugins/EnhancedInput/Source/EnhancedInput/Private/EnhancedPlayerInput.cpp
+++ b/Engine/Plugins/EnhancedInput/Source/EnhancedInput/Private/EnhancedPlayerInput.cpp
@@ -397,6 +397,12 @@ void UEnhancedPlayerInput::EvaluateInputDelegates(const TArray<UInputComponent*>
 				}
 			}
 		}
+		// Virtual Keys have different values which are set at runtime depending on the platform. We need to make sure
+		// that we are grabbing the actual FKey name that it represents
+		else if (Mapping.Key.IsVirtual())
+		{
+			KeyState = GetKeyState(Mapping.Key.GetVirtualKey());
+		}
 		else
 		{
 			KeyState = GetKeyState(Mapping.Key);
diff --git a/Engine/Plugins/Runtime/CommonUI/Source/CommonUI/Private/Input/CommonAnalogCursor.cpp b/Engine/Plugins/Runtime/CommonUI/Source/CommonUI/Private/Input/CommonAnalogCursor.cpp
--- a/Engine/Plugins/Runtime/CommonUI/Source/CommonUI/Private/Input/CommonAnalogCursor.cpp
+++ b/Engine/Plugins/Runtime/CommonUI/Source/CommonUI/Private/Input/CommonAnalogCursor.cpp
@@ -255,7 +255,7 @@ bool FCommonAnalogCursor::HandleKeyDownEvent(FSlateApplication& SlateApp, const
 #endif
 
 		// We support binding actions to the virtual accept key, so it's a special flower that gets processed right now
-		const bool bIsVirtualAccept = InKeyEvent.GetKey() == EKeys::Virtual_Accept;
+		const bool bIsVirtualAccept = InKeyEvent.GetKey() == EKeys::Virtual_Gamepad_Accept.GetVirtualKey();
 		const EInputEvent InputEventType = InKeyEvent.IsRepeat() ? IE_Repeat : IE_Pressed;
 		if (bIsVirtualAccept && ActionRouter.ProcessInput(InKeyEvent.GetKey(), InputEventType) == ERouteUIInputResult::Handled)
 		{
@@ -290,7 +290,7 @@ bool FCommonAnalogCursor::HandleKeyUpEvent(FSlateApplication& SlateApp, const FK
 		TGuardValue<TOptional<FKeyEvent>> KeyUpEventGuard(ActiveKeyUpEvent, TOptional<FKeyEvent>(InKeyEvent));
 		
 		// We support binding actions to the virtual accept key, so it's a special flower that gets processed right now
-		const bool bIsVirtualAccept = InKeyEvent.GetKey() == EKeys::Virtual_Accept;
+		const bool bIsVirtualAccept = InKeyEvent.GetKey() == EKeys::Virtual_Gamepad_Accept.GetVirtualKey();
 		if (bIsVirtualAccept && ActionRouter.ProcessInput(InKeyEvent.GetKey(), IE_Released) == ERouteUIInputResult::Handled)
 		{
 			return true;
@@ -419,7 +419,7 @@ bool FCommonAnalogCursor::ShouldVirtualAcceptSimulateMouseButton(const FKeyEvent
 
 void FCommonAnalogCursor::OnVirtualAcceptHoldCanceled()
 {
-	if (ActiveKeyUpEvent.IsSet() && ActiveKeyUpEvent->GetKey() == EKeys::Virtual_Accept && ShouldVirtualAcceptSimulateMouseButton(ActiveKeyUpEvent.GetValue(), IE_Pressed))
+	if (ActiveKeyUpEvent.IsSet() && ActiveKeyUpEvent->GetKey() == EKeys::Virtual_Gamepad_Accept.GetVirtualKey() && ShouldVirtualAcceptSimulateMouseButton(ActiveKeyUpEvent.GetValue(), IE_Pressed))
 	{
 		UCommonInputSubsystem& InputSubsystem = ActionRouter.GetInputSubsystem();
 		InputSubsystem.SetIsGamepadSimulatedClick(true);
@@ -457,7 +457,7 @@ EOrientation FCommonAnalogCursor::DetermineScrollOrientation(const UWidget& Widg
 
 bool FCommonAnalogCursor::IsRelevantInput(const FKeyEvent& KeyEvent) const
 {
-	return IsUsingGamepad() && FAnalogCursor::IsRelevantInput(KeyEvent) && (IsGameViewportInFocusPathWithoutCapture() || (KeyEvent.GetKey() == EKeys::Virtual_Accept && CanReleaseMouseCapture()));
+	return IsUsingGamepad() && FAnalogCursor::IsRelevantInput(KeyEvent) && (IsGameViewportInFocusPathWithoutCapture() || (KeyEvent.GetKey() == EKeys::Virtual_Gamepad_Accept.GetVirtualKey() && CanReleaseMouseCapture()));
 }
 
 bool FCommonAnalogCursor::IsRelevantInput(const FAnalogInputEvent& AnalogInputEvent) const
diff --git a/Engine/Plugins/Runtime/CommonUI/Source/CommonUI/Private/Input/CommonBoundActionBar.cpp b/Engine/Plugins/Runtime/CommonUI/Source/CommonUI/Private/Input/CommonBoundActionBar.cpp
--- a/Engine/Plugins/Runtime/CommonUI/Source/CommonUI/Private/Input/CommonBoundActionBar.cpp
+++ b/Engine/Plugins/Runtime/CommonUI/Source/CommonUI/Private/Input/CommonBoundActionBar.cpp
@@ -233,7 +233,7 @@ void UCommonBoundActionBar::HandleDeferredDisplayUpdate()
 										}
 									}
 
-									return Key == EKeys::Virtual_Back || Key == EKeys::Escape || Key == EKeys::Android_Back;
+									return Key == EKeys::Virtual_Gamepad_Back.GetVirtualKey() || Key == EKeys::Escape || Key == EKeys::Android_Back;
 								}
 								else if (InputAction)
 								{
@@ -248,7 +248,7 @@ void UCommonBoundActionBar::HandleDeferredDisplayUpdate()
 										}
 									}
 
-									return Key == EKeys::Virtual_Back || Key == EKeys::Escape || Key == EKeys::Android_Back;
+									return Key == EKeys::Virtual_Gamepad_Back.GetVirtualKey() || Key == EKeys::Escape || Key == EKeys::Android_Back;
 								}
 
 								return false;
diff --git a/Engine/Plugins/Runtime/CommonUI/Source/CommonUI/Private/Input/CommonUIActionRouterBase.cpp b/Engine/Plugins/Runtime/CommonUI/Source/CommonUI/Private/Input/CommonUIActionRouterBase.cpp
--- a/Engine/Plugins/Runtime/CommonUI/Source/CommonUI/Private/Input/CommonUIActionRouterBase.cpp
+++ b/Engine/Plugins/Runtime/CommonUI/Source/CommonUI/Private/Input/CommonUIActionRouterBase.cpp
@@ -598,7 +598,7 @@ ERouteUIInputResult UCommonUIActionRouterBase::ProcessInput(FKey Key, EInputEven
 			// A hold action was in progress but quickly aborted, so we want to generate a press action now for any normal bindings that are interested
 			if (!ProcessNormalInputFunc(ActionRouter, IE_Pressed))
 			{
-				if (Key == EKeys::Virtual_Accept && ActionRouter.AnalogCursor.IsValid())
+				if (Key == EKeys::Virtual_Gamepad_Accept.GetVirtualKey() && ActionRouter.AnalogCursor.IsValid())
 				{
 					// If Virtual_Accept is bound to a Hold action, the Pressed event for it will not trigger a simulated click
 					// but when a hold action is canceled fast enough, we want to trigger any "normal" actions bound to the same key
diff --git a/Engine/Source/Programs/AutomationTool/Gauntlet/Framework/Utils/Gauntlet.InputActions.cs b/Engine/Source/Programs/AutomationTool/Gauntlet/Framework/Utils/Gauntlet.InputActions.cs
--- a/Engine/Source/Programs/AutomationTool/Gauntlet/Framework/Utils/Gauntlet.InputActions.cs
+++ b/Engine/Source/Programs/AutomationTool/Gauntlet/Framework/Utils/Gauntlet.InputActions.cs
@@ -533,11 +533,7 @@ namespace Gauntlet.Utils
 		public static FKey ValveIndex_Right_Trackpad_Down = new FKey("ValveIndex_Right_Trackpad_Down");
 		public static FKey ValveIndex_Right_Trackpad_Left = new FKey("ValveIndex_Right_Trackpad_Left");
 		public static FKey ValveIndex_Right_Trackpad_Right = new FKey("ValveIndex_Right_Trackpad_Right");
-
-		/*FKey Virtual_Accept;
-		FKey Virtual_Back;
-
-		FKey Invalid;*/
+		
 		public EKeys()
 		{
 			UnrealTargetPlatform.TryParse(Globals.Params.ParseValue("Platform", "Win64"), out UnrealTargetPlatform Platform);
diff --git a/Engine/Source/Runtime/InputCore/Classes/InputCoreTypes.h b/Engine/Source/Runtime/InputCore/Classes/InputCoreTypes.h
--- a/Engine/Source/Runtime/InputCore/Classes/InputCoreTypes.h
+++ b/Engine/Source/Runtime/InputCore/Classes/InputCoreTypes.h
@@ -80,6 +80,7 @@ struct FKey
 	INPUTCORE_API bool IsAxis3D() const;
 	INPUTCORE_API bool IsDigital() const;
 	INPUTCORE_API bool IsAnalog() const;
+	INPUTCORE_API bool IsVirtual() const;
 	INPUTCORE_API bool IsBindableInBlueprints() const;
 	INPUTCORE_API bool ShouldUpdateAxisWithoutSamples() const;
 	INPUTCORE_API bool IsBindableToActions() const;
@@ -91,6 +92,13 @@ struct FKey
 	INPUTCORE_API FName GetMenuCategory() const;
 	INPUTCORE_API EPairedAxis GetPairedAxis() const;
 	INPUTCORE_API FKey GetPairedAxisKey() const;
+	INPUTCORE_API FKey GetVirtualKey() const;
+
+	/**
+	 * @param Other		The key to compare against
+	 * @return			True if this FKey's name or its virtual key name (if it has one) is the same as the other.
+	 */
+	INPUTCORE_API bool IsSameResolvedKey(const FKey& Other) const;
 
 	INPUTCORE_API bool SerializeFromMismatchedTag(struct FPropertyTag const& Tag, FStructuredArchive::FSlot Slot);
 	INPUTCORE_API bool ExportTextItem(FString& ValueStr, FKey const& DefaultValue, UObject* Parent, int32 PortFlags, UObject* ExportRootScope) const;
@@ -156,6 +164,14 @@ struct FKeyDetails
 		// Gesture input types such as Flick, Pinch, and Rotate
 		Gesture					 = 1 << 12,
 
+		// A key is "Virtual" if it is an abstract key whose actual value may change dependent on the platform.
+		// For example, the standard "Accept" button on some platforms may be Gamepad_FaceButton_Down, while on
+		// other platforms it may be Gamepad_FaceButton_Right.
+		//
+		// Virtual keys are typically the most useful when it comes to building UI related code because they allow you 
+		// to create experiences and screens displayed to the end user which are consistent with the platform's UI.
+		Virtual					 = 1 << 13,
+
 		NoFlags                  = 0,
 	};
 
@@ -177,6 +193,7 @@ struct FKeyDetails
 	inline bool IsBindableToActions() const { return bIsBindableToActions != 0; }
 	inline bool IsGesture() const { return bIsGesture != 0; }
 	inline bool IsDeprecated() const { return bIsDeprecated != 0; }
+	inline bool IsVirtual() const { return bIsVirtual != 0; }
 	inline FName GetMenuCategory() const { return MenuCategory; }
 	INPUTCORE_API FText GetDisplayName(const bool bLongDisplayName = true) const;
 	inline const FKey& GetKey() const { return Key; }
@@ -184,6 +201,11 @@ struct FKeyDetails
 	// Key pairing
 	inline EPairedAxis GetPairedAxis() const { return PairedAxis; }
 	inline const FKey& GetPairedAxisKey() const { return PairedAxisKey; }
+	
+	/**
+	 * For virtual Keys, this will return the virtual key value. For non-virtual keys, it will return the "normal" key for easier checking
+	 */
+	inline const FKey& GetVirtualKey() const { return bIsVirtual ? VirtualKeyValue : Key; }
 
 private:
 	friend struct EKeys;
@@ -205,6 +227,11 @@ private:
 	EPairedAxis PairedAxis = EPairedAxis::Unpaired;		// Paired axis identifier. Lets this key know which axis it represents on the PairedAxisKey
 	FKey		PairedAxisKey;							// Paired axis reference. This is the FKey representing the final paired vector axis. Note: NOT the other key in the pairing.
 
+	/**
+	 * The actual FKey which should be used at runtime for virtual FKeys
+	 */
+	FKey VirtualKeyValue;
+
 	FName MenuCategory;
 
 	uint8 bIsModifierKey : 1;
@@ -215,6 +242,7 @@ private:
 	uint8 bShouldUpdateAxisWithoutSamples : 1;
 	uint8 bIsBindableToActions : 1;
 	uint8 bIsDeprecated : 1;
+	uint8 bIsVirtual : 1;
 	uint8 bIsGesture : 1;
 	EInputAxisType AxisType;
 
@@ -653,8 +681,14 @@ struct EKeys
 	static INPUTCORE_API const FKey ValveIndex_Right_Trackpad_Right;
 
 	// Virtual buttons that use other buttons depending on the platform
+
+	UE_DEPRECATED(5.7, "Use Virtual_Gamepad_Accept.GetVirtualKey() instead")
 	static INPUTCORE_API const FKey Virtual_Accept;
+	UE_DEPRECATED(5.7, "Use Virtual_Gamepad_Back.GetVirtualKey() instead")
 	static INPUTCORE_API const FKey Virtual_Back;
+	
+	static INPUTCORE_API const FKey Virtual_Gamepad_Accept;
+	static INPUTCORE_API const FKey Virtual_Gamepad_Back;
 
 	static INPUTCORE_API const FKey Invalid;
 
@@ -669,11 +703,13 @@ struct EKeys
 	static INPUTCORE_API const FName NAME_KeyboardCategory;
 	static INPUTCORE_API const FName NAME_GamepadCategory;
 	static INPUTCORE_API const FName NAME_MouseCategory;
+	static INPUTCORE_API const FName NAME_VirtualCategory;
 
 	static INPUTCORE_API void Initialize();
 	static INPUTCORE_API void AddKey(const FKeyDetails& KeyDetails);
 	// Map the two provided keys to the X and Y axes of the paired key
 	static INPUTCORE_API void AddPairedKey(const FKeyDetails& PairedKeyDetails, FKey KeyX, FKey KeyY);
+	static INPUTCORE_API void AddVirtualKey(const FKeyDetails& VirtualKeyDetails, const FKey& VirtualKeyValue);
 	static INPUTCORE_API void GetAllKeys(TArray<FKey>& OutKeys);
 	static INPUTCORE_API TSharedPtr<FKeyDetails> GetKeyDetails(const FKey Key);
 	static INPUTCORE_API void RemoveKeysWithCategory(const FName InCategory);
diff --git a/Engine/Source/Runtime/InputCore/Private/InputCoreTypes.cpp b/Engine/Source/Runtime/InputCore/Private/InputCoreTypes.cpp
--- a/Engine/Source/Runtime/InputCore/Private/InputCoreTypes.cpp
+++ b/Engine/Source/Runtime/InputCore/Private/InputCoreTypes.cpp
@@ -413,9 +413,16 @@ const FKey EKeys::Invalid(NAME_None);
 const FName EKeys::NAME_GamepadCategory("Gamepad");
 const FName EKeys::NAME_MouseCategory("Mouse");
 const FName EKeys::NAME_KeyboardCategory("Key");
+const FName EKeys::NAME_VirtualCategory("Virtual");
 
+// Create the virtual FKeys from the platform specified versions of their keys
+PRAGMA_DISABLE_DEPRECATION_WARNINGS
 const FKey EKeys::Virtual_Accept = FPlatformInput::GetGamepadAcceptKey();
 const FKey EKeys::Virtual_Back = FPlatformInput::GetGamepadBackKey();
+PRAGMA_ENABLE_DEPRECATION_WARNINGS
+
+const FKey EKeys::Virtual_Gamepad_Accept("Virtual_Gamepad_Accept");
+const FKey EKeys::Virtual_Gamepad_Back("Virtual_Gamepad_Back");
 
 bool EKeys::bInitialized = false;
 TMap<FKey, TSharedPtr<FKeyDetails>> EKeys::InputKeys;
@@ -451,6 +458,7 @@ void FKeyDetails::CommonInit(const uint32 InKeyFlags)
 	bIsBindableToActions = ((~InKeyFlags & EKeyFlags::NotActionBindableKey) != 0) && ((~InKeyFlags & EKeyFlags::Deprecated) != 0);
 	bIsDeprecated = ((InKeyFlags & EKeyFlags::Deprecated) != 0);
 	bIsGesture = ((InKeyFlags & EKeyFlags::Gesture) != 0);
+	bIsVirtual = ((InKeyFlags & EKeyFlags::Virtual) != 0);
 
 	if ((InKeyFlags & EKeyFlags::ButtonAxis) != 0)
 	{
@@ -479,7 +487,11 @@ void FKeyDetails::CommonInit(const uint32 InKeyFlags)
 	// Set up default menu categories
 	if (MenuCategory.IsNone())
 	{
-		if (IsGamepadKey())
+		if (IsVirtual())
+		{
+			MenuCategory = EKeys::NAME_VirtualCategory;
+		}
+		else if (IsGamepadKey())
 		{
 			MenuCategory = EKeys::NAME_GamepadCategory;
 		}
@@ -511,6 +523,7 @@ void EKeys::Initialize()
 	}
 
 	AddMenuCategoryDisplayInfo(NAME_GamepadCategory, LOCTEXT("GamepadSubCategory", "Gamepad"), TEXT("GraphEditor.PadEvent_16x"));
+	AddMenuCategoryDisplayInfo(NAME_VirtualCategory, LOCTEXT("VirtualSubCategory", "Virtual"), TEXT("GraphEditor.PadEvent_16x"));
 	AddMenuCategoryDisplayInfo(NAME_MouseCategory, LOCTEXT("MouseSubCategory", "Mouse"), TEXT("GraphEditor.MouseEvent_16x"));
 	AddMenuCategoryDisplayInfo(NAME_KeyboardCategory, LOCTEXT("KeyboardSubCategory", "Keyboard"), TEXT("GraphEditor.KeyEvent_16x"));
 
@@ -708,6 +721,10 @@ void EKeys::Initialize()
 	AddKey(FKeyDetails(EKeys::Gamepad_LeftThumbstick, TAttribute<FText>::Create(TAttribute<FText>::FGetter::CreateStatic(&EKeys::GetGamepadDisplayName, EKeys::Gamepad_LeftThumbstick)), FKeyDetails::GamepadKey));
 	AddKey(FKeyDetails(EKeys::Gamepad_RightThumbstick, TAttribute<FText>::Create(TAttribute<FText>::FGetter::CreateStatic(&EKeys::GetGamepadDisplayName, EKeys::Gamepad_RightThumbstick)), FKeyDetails::GamepadKey));
 
+	// The virtual accept and back keys represent platform specific actions which may have different bindings per-platform.
+	AddVirtualKey(FKeyDetails(EKeys::Virtual_Gamepad_Accept, LOCTEXT("Virtual_Gamepad_Accept_Key", "Virtual Gamepad Accept"), FKeyDetails::GamepadKey | FKeyDetails::Virtual, EKeys::NAME_VirtualCategory), FPlatformInput::GetGamepadAcceptKey());
+	AddVirtualKey(FKeyDetails(EKeys::Virtual_Gamepad_Back, LOCTEXT("Virtual_Gamepad_Back_Key", "Virtual Gamepad Back"), FKeyDetails::GamepadKey | FKeyDetails::Virtual, EKeys::NAME_VirtualCategory), FPlatformInput::GetGamepadBackKey());
+	
 	AddMenuCategoryDisplayInfo("Motion", LOCTEXT("MotionSubCateogry", "Motion"), TEXT("GraphEditor.KeyEvent_16x"));
 
 	AddKey(FKeyDetails(EKeys::Tilt, LOCTEXT("Tilt", "Tilt"), FKeyDetails::Axis3D, "Motion"));
@@ -991,6 +1008,23 @@ void EKeys::AddPairedKey(const FKeyDetails& PairedKeyDetails, FKey KeyX, FKey Ke
 	); 
 }
 
+void EKeys::AddVirtualKey(const FKeyDetails& VirtualKeyDetails, const FKey& VirtualKeyValue)
+{
+	// Validate the given virtual key meets all the necessary conditions.
+	if (!ensureMsgf(VirtualKeyDetails.IsVirtual(), TEXT("Paired key '%s' must be marked as a virtual key"), *VirtualKeyDetails.GetKey().ToString()) ||
+		!ensureMsgf(InputKeys.Contains(VirtualKeyValue), TEXT("Failed to locate key '%s' for virtual key creation. Make sure you've bound it with AddKey before calling AddVirtualKey."), *VirtualKeyValue.ToString())
+		)
+	{
+		return;
+	}
+	
+	// Add the basic paired key information
+	AddKey(VirtualKeyDetails);
+
+	// Set up the virtual FKey value to be the one that was provided
+	InputKeys[VirtualKeyDetails.GetKey()]->VirtualKeyValue = VirtualKeyValue;
+}
+
 void EKeys::GetAllKeys(TArray<FKey>& OutKeys)
 {
 	InputKeys.GetKeys(OutKeys);
@@ -1337,6 +1371,12 @@ bool FKey::IsAnalog() const
 	return (KeyDetails.IsValid() ? KeyDetails->IsAnalog() : false);
 }
 
+bool FKey::IsVirtual() const
+{
+	ConditionalLookupKeyDetails();
+	return (KeyDetails.IsValid() ? KeyDetails->IsVirtual() : false);
+}
+
 bool FKey::IsBindableInBlueprints() const
 {
 	ConditionalLookupKeyDetails();
@@ -1396,6 +1436,12 @@ FKey FKey::GetPairedAxisKey() const
 	return (KeyDetails.IsValid() ? KeyDetails->GetPairedAxisKey() : FKey());
 }
 
+FKey FKey::GetVirtualKey() const
+{
+	ConditionalLookupKeyDetails();
+	return (KeyDetails.IsValid() ? KeyDetails->GetVirtualKey() : FKey());
+}
+
 void FKey::ConditionalLookupKeyDetails() const
 {
 	if (!KeyDetails.IsValid())
@@ -1404,6 +1450,13 @@ void FKey::ConditionalLookupKeyDetails() const
 	}
 }
 
+bool FKey::IsSameResolvedKey(const FKey& Other) const
+{
+	// If a key is virtual, then it's FName will be different from what "KeyName" is.
+	// If it is not virtual, then GetVirtualKey simply returns the normal Key implementation
+	return GetVirtualKey() == Other.GetVirtualKey();
+}
+
 bool FKey::SerializeFromMismatchedTag(struct FPropertyTag const& Tag, FStructuredArchive::FSlot Slot)
 {
 	if (Tag.GetType().IsEnum(TEXT("EKeys")))
diff --git a/Engine/Source/Runtime/Slate/Private/Framework/Application/AnalogCursor.cpp b/Engine/Source/Runtime/Slate/Private/Framework/Application/AnalogCursor.cpp
--- a/Engine/Source/Runtime/Slate/Private/Framework/Application/AnalogCursor.cpp
+++ b/Engine/Source/Runtime/Slate/Private/Framework/Application/AnalogCursor.cpp
@@ -52,7 +52,7 @@ bool FAnalogCursor::HandleKeyDownEvent(FSlateApplication& SlateApp, const FKeyEv
 		}
 
 		// Bottom face button is a click
-		if (Key == EKeys::Virtual_Accept)
+		if (Key == EKeys::Virtual_Gamepad_Accept.GetVirtualKey())
 		{
 			if (!InKeyEvent.IsRepeat())
 			{
@@ -98,7 +98,7 @@ bool FAnalogCursor::HandleKeyUpEvent(FSlateApplication& SlateApp, const FKeyEven
 		}
 
 		// Bottom face button is a click
-		if (Key == EKeys::Virtual_Accept && !InKeyEvent.IsRepeat())
+		if (Key == EKeys::Virtual_Gamepad_Accept.GetVirtualKey() && !InKeyEvent.IsRepeat())
 		{
 			if (TSharedPtr<FSlateUser> SlateUser = SlateApp.GetUser(InKeyEvent))
 			{
diff --git a/Engine/Source/Runtime/Slate/Private/Framework/Application/NavigationConfig.cpp b/Engine/Source/Runtime/Slate/Private/Framework/Application/NavigationConfig.cpp
--- a/Engine/Source/Runtime/Slate/Private/Framework/Application/NavigationConfig.cpp
+++ b/Engine/Source/Runtime/Slate/Private/Framework/Application/NavigationConfig.cpp
@@ -31,11 +31,11 @@ FNavigationConfig::FNavigationConfig()
 	// By default, enter, space, and gamepad accept are all counted as accept
 	KeyActionRules.Emplace(EKeys::Enter, EUINavigationAction::Accept);
 	KeyActionRules.Emplace(EKeys::SpaceBar, EUINavigationAction::Accept);
-	KeyActionRules.Emplace(EKeys::Virtual_Accept, EUINavigationAction::Accept);
+	KeyActionRules.Emplace(EKeys::Virtual_Gamepad_Accept.GetVirtualKey(), EUINavigationAction::Accept);
 
 	// By default, escape and gamepad back count as leaving current scope
 	KeyActionRules.Emplace(EKeys::Escape, EUINavigationAction::Back);
-	KeyActionRules.Emplace(EKeys::Virtual_Back, EUINavigationAction::Back);
+	KeyActionRules.Emplace(EKeys::Virtual_Gamepad_Back.GetVirtualKey(), EUINavigationAction::Back);
 }
 
 FNavigationConfig::~FNavigationConfig()
diff --git a/Engine/Source/Runtime/SlateCore/Public/Types/SlateEnums.h b/Engine/Source/Runtime/SlateCore/Public/Types/SlateEnums.h
--- a/Engine/Source/Runtime/SlateCore/Public/Types/SlateEnums.h
+++ b/Engine/Source/Runtime/SlateCore/Public/Types/SlateEnums.h
@@ -122,10 +122,10 @@ ENUM_RANGE_BY_COUNT(EUINavigation, EUINavigation::Num);
 UENUM(BlueprintType)
 enum class EUINavigationAction : uint8
 {
-	/** Accept, normally acts like a mouse click. Corresponds to Virtual_Accept key */
+	/** Accept, normally acts like a mouse click. Corresponds to Virtual_Gamepad_Accept key */
 	Accept,
 
-	/** Back/Cancel, normally acts like escape. Corresponds to Virtual_Back key */
+	/** Back/Cancel, normally acts like escape. Corresponds to Virtual_Gamepad_Back key */
 	Back,
 
 	/** Number of navigation actions */
